MQTT plugin added to engine

diff --git a/Engine/Plugins/Protocols/MQTT/MQTT.uplugin b/Engine/Plugins/Protocols/MQTT/MQTT.uplugin
new file mode 100644
--- /dev/null
+++ b/Engine/Plugins/Protocols/MQTT/MQTT.uplugin
@@ -0,0 +1,34 @@
+{
+	"FileVersion": 3,
+	"Version": 1.0,
+	"VersionName": "1.0",
+	"FriendlyName": "MQTT",
+	"Description": "MQTT broker and client",
+	"Category": "IOT",
+	"CreatedBy": "Epic Games, Inc.",
+	"CreatedByURL": "http://epicgames.com",
+	"EnabledByDefault": false,
+	"CanContainContent": false,
+	"IsBetaVersion": true,
+	"SupportedTargetPlatforms": [
+		"Win64",
+		"Mac",
+		"Linux"
+	],
+	"Modules": [
+		{
+			"Name": "MQTTCore",
+			"Type": "Runtime"
+		},
+		{
+			"Name": "MQTTCoreEditor",
+			"Type": "Editor"
+		}
+	],
+	"Plugins": [
+		{
+			"Name": "JsonBlueprintUtilities",
+			"Enabled": true
+		}
+	]
+}
diff --git a/Engine/Plugins/Protocols/MQTT/Source/MQTTCore/MQTTCore.Build.cs b/Engine/Plugins/Protocols/MQTT/Source/MQTTCore/MQTTCore.Build.cs
new file mode 100644
--- /dev/null
+++ b/Engine/Plugins/Protocols/MQTT/Source/MQTTCore/MQTTCore.Build.cs
@@ -0,0 +1,40 @@
+// Copyright Epic Games, Inc. All Rights Reserved.
+
+using UnrealBuildTool;
+
+public class MQTTCore : ModuleRules
+{
+	public MQTTCore(ReadOnlyTargetRules Target) : base(Target)
+	{
+		PublicDependencyModuleNames.AddRange(
+			new[]
+			{
+				"Core",
+				"CoreUObject",
+				"JsonBlueprintUtilities",
+				"JsonUtilities",
+			}
+		);
+
+		PrivateDependencyModuleNames.AddRange(
+			new[]
+			{
+				"Engine",
+				"Json",
+				"JsonBlueprintUtilities",
+				"Networking",
+				"Projects",
+				"Sockets"
+			}
+		);
+		
+		if (Target.bBuildEditor)
+		{
+			PrivateDependencyModuleNames.AddRange(
+				new[] {
+					"FunctionalTesting",
+					"UnrealEd"
+				});
+		}
+	}
+}
diff --git a/Engine/Plugins/Protocols/MQTT/Source/MQTTCore/Private/MQTTClient.cpp b/Engine/Plugins/Protocols/MQTT/Source/MQTTCore/Private/MQTTClient.cpp
new file mode 100644
--- /dev/null
+++ b/Engine/Plugins/Protocols/MQTT/Source/MQTTCore/Private/MQTTClient.cpp
@@ -0,0 +1,388 @@
+﻿// Copyright Epic Games, Inc. All Rights Reserved.
+
+#include "MQTTClient.h"
+
+#include "IMQTTCoreModule.h"
+#include "MQTTClientMessage.h"
+#include "MQTTShared.h"
+#include "MQTTCoreLog.h"
+#include "SocketSubsystem.h"
+#include "Algo/ForEach.h"
+#include "Async/Async.h"
+#include "Templates/SharedPointer.h"
+
+FMQTTClient::FMQTTClient(const FMQTTURL& InURL)
+	: ClientId(InURL.ToGuid())
+	, URL(InURL)
+{
+	FText URLMessage;
+	if(!InURL.IsValid(URLMessage))
+	{
+		UE_LOG(LogMQTTCore, Error, TEXT("URL is not valid: %s"), *URLMessage.ToString());
+		return;
+	}
+
+	ISocketSubsystem* SocketSubsystem = ISocketSubsystem::Get(PLATFORM_SOCKETSUBSYSTEM);
+	TSharedPtr<FInternetAddr> InternetAddr;
+	if(!URL.GetAddress(SocketSubsystem, InternetAddr))
+	{
+		UE_LOG(LogMQTTCore, Error, TEXT("Cannot create MQTTConnection: Invalid IP address: %s"), *InURL.Host);
+		URL.SetInvalid();
+	}
+	else
+	{
+		Connection = FMQTTConnection::TryCreate(InternetAddr);
+		Connection->OnMessage().BindRaw(this, &FMQTTClient::OnMessagePacket);
+	} 
+}
+
+FMQTTClient::~FMQTTClient()
+{
+	if(Connection.IsValid())
+	{
+		Connection->OnMessage().Unbind();
+		Connection.Reset();
+	}
+}
+
+TFuture<EMQTTConnectReturnCode> FMQTTClient::Connect(const bool bCleanSession)
+{
+	TFuture<EMQTTConnectReturnCode> Result = MakeFulfilledPromise<EMQTTConnectReturnCode>(URL.IsValid() ? EMQTTConnectReturnCode::AlreadyConnected : EMQTTConnectReturnCode::InvalidURL).GetFuture();;
+	if(Connection.IsValid())
+	{
+		const auto ConnectionState = Connection->GetState();
+		if(Connection->IsConnected())
+		{
+			Result = MakeFulfilledPromise<EMQTTConnectReturnCode>(EMQTTConnectReturnCode::AlreadyConnected).GetFuture();
+		}
+		else if(ConnectionState == EMQTTState::None)
+		{
+			Result = MakeFulfilledPromise<EMQTTConnectReturnCode>(EMQTTConnectReturnCode::SocketError).GetFuture();
+		}
+		// Not currently connecting
+		else if(ConnectionState != EMQTTState::Connecting)
+		{
+			// Set immediately to prevent multiple requests
+			Connection->SetState(EMQTTState::Connecting);
+			UE_LOG(LogTemp, VeryVerbose, TEXT("QUEUED CONNECT"));
+			Result = Connection->QueueOperation<FMQTTConnectOperation>(
+					FMQTTConnectPacket(
+						URL.ToGuid().ToString(),
+						60,
+						URL.Username.IsEmpty() ? "" : URL.Username,
+						URL.Password.IsEmpty() ? "" : URL.Password,
+						bCleanSession))
+				.Next([Self = TWeakPtr<FMQTTClient, ESPMode::ThreadSafe>(AsShared())](const FMQTTConnectAckPacket InResponse)
+				{
+					UE_LOG(LogTemp, VeryVerbose, TEXT("GOT CONNACK"));
+					if(!Self.IsValid())
+					{
+						return EMQTTConnectReturnCode::SocketError;
+					}
+					
+					if (InResponse.ReturnCode > EMQTTConnectReturnCode::Accepted && InResponse.ReturnCode != EMQTTConnectReturnCode::AlreadyConnected)
+					{
+						UE_LOG(LogMQTTCore, Warning, TEXT("Couldn't connect to MQTT server: %s"), MQTT::GetMQTTConnectReturnCodeDescription(InResponse.ReturnCode));
+					}
+
+					Self.Pin()->Connection->SetState(EMQTTState::Connected);
+					return InResponse.ReturnCode;
+				});
+		}
+	}
+
+	return Result.Next([Self = TSharedPtr<FMQTTClient, ESPMode::ThreadSafe>(AsShared())](EMQTTConnectReturnCode InCode)
+	{
+		if(!Self.IsValid())
+		{
+			return EMQTTConnectReturnCode::SocketError;
+		}
+		
+		Self->OnConnect().Broadcast(InCode);
+		return InCode;
+	});
+}
+
+TFuture<void> FMQTTClient::Disconnect()
+{
+	if(Connection.IsValid())
+	{
+		// Check if connecting/connecting
+		if(Connection->GetState() > EMQTTState::Connecting)
+		{
+			return MakeFulfilledPromise<void>().GetFuture();
+		}
+	
+		// Set immediately to prevent multiple requests
+		Connection->SetState(EMQTTState::Disconnecting);
+		return Connection->QueueOperation<FMQTTDisconnectOperation>(FMQTTDisconnectPacket())
+			.Next([Self = TWeakPtr<FMQTTClient, ESPMode::ThreadSafe>(AsShared())](int32)
+			{
+				if (!Self.IsValid())
+				{
+					return;
+				}
+
+				if(const auto SharedSelf = Self.Pin())
+				{
+					SharedSelf->Connection->SetState(EMQTTState::Disconnected);
+					SharedSelf->OnDisconnect().Broadcast();
+				}
+			});
+	}
+
+	return MakeFulfilledPromise<void>().GetFuture();
+}
+
+TFuture<bool> FMQTTClient::Publish(const FString& InTopic, const FString& InPayload, const EMQTTQualityOfService InQoS, const bool bInRetain)
+{
+	return Publish(InTopic, TArray<uint8>((uint8*)TCHAR_TO_UTF8(*InPayload), InPayload.Len()), InQoS, bInRetain);
+}
+
+TFuture<bool> FMQTTClient::Publish(const FString& InTopic, const TArray<uint8>& InPayload, const EMQTTQualityOfService InQoS, const bool bInRetain)
+{
+	if(Connection.IsValid())
+	{
+		FMQTTPublishPacket PublishPacket(
+			Connection->PopId(),
+			InTopic,
+			InPayload,
+			InQoS,
+			bInRetain);
+
+		if(PublishPacket.Topic.IsEmpty())
+		{
+			UE_LOG(LogMQTTCore, Warning, TEXT("Attempted to publish message without topic"));
+			return MakeFulfilledPromise<bool>(false).GetFuture();			
+		}
+
+		if(!FMQTTTopic::IsValid(InTopic))
+		{
+			UE_LOG(LogMQTTCore, Warning, TEXT("Attempted to publish with an invalid or malformed topic: \'%s\'"), *InTopic);
+			return MakeFulfilledPromise<bool>(false).GetFuture();			
+		}
+
+		TFuture<FMQTTPublishAckPacket> PublishResponse;
+		switch(InQoS)
+		{
+		case EMQTTQualityOfService::Once:
+			return Connection->QueueOperation<FMQTTPublishOperationQoS0>(MoveTemp(PublishPacket))
+			.Next([Self = TWeakPtr<FMQTTClient, ESPMode::ThreadSafe>(AsShared())](auto /*InResponse*/)
+			{
+				if(!Self.IsValid())
+				{
+					return false;
+				} 
+				
+				Self.Pin()->OnPublish().Broadcast();
+				return true;
+			});
+			
+		case EMQTTQualityOfService::AtLeastOnce:
+			return Connection->QueueOperation<FMQTTPublishOperationQoS1>(MoveTemp(PublishPacket))
+			.Next([Self = TWeakPtr<FMQTTClient, ESPMode::ThreadSafe>(AsShared())](FMQTTPublishAckPacket /*InResponse*/)
+			{
+				if(!Self.IsValid())
+				{
+					return false;
+				} 
+				
+				Self.Pin()->OnPublish().Broadcast();
+				return true;
+			});
+			
+		case EMQTTQualityOfService::ExactlyOnce:
+			return Connection->QueueOperation<FMQTTPublishOperationQoS2>(MoveTemp(PublishPacket))
+			.Next([Self = TWeakPtr<FMQTTClient, ESPMode::ThreadSafe>(AsShared())](FMQTTPublishCompletePacket /*InResponse*/)
+			{
+				if(!Self.IsValid())
+				{
+					return false;
+				} 
+				
+				Self.Pin()->OnPublish().Broadcast();
+				return true;
+			});
+			
+		default: ;
+		}
+	}
+
+	return MakeFulfilledPromise<bool>(false).GetFuture();
+}
+
+TFuture<TArray<FMQTTSubscribeResult>> FMQTTClient::Subscribe(const TArray<TPair<FString, EMQTTQualityOfService>>& InTopicFilterQoSPairs)
+{
+	static TArray<FMQTTSubscribeResult> FailureResponse = { FMQTTSubscribeResult(EMQTTSubscribeReturnCode::Failure, nullptr) };
+	
+	if(InTopicFilterQoSPairs.IsEmpty())
+	{
+		UE_LOG(LogMQTTCore, Warning, TEXT("Attempted to subscribe without specifying a topic"));		
+		return MakeFulfilledPromise<TArray<FMQTTSubscribeResult>>(FailureResponse).GetFuture();
+	}
+	
+	if(Connection.IsValid())
+	{
+		TArray<FString> TopicFilters;
+		TopicFilters.Reserve(InTopicFilterQoSPairs.Num());
+
+		TArray<EMQTTQualityOfService> TopicQoS;
+		TopicQoS.Reserve(InTopicFilterQoSPairs.Num());
+
+		for(const TPair<FString, EMQTTQualityOfService>& TopicFilterQoSPair : InTopicFilterQoSPairs)
+		{
+			if(TopicFilterQoSPair.Key.IsEmpty())
+			{
+				UE_LOG(LogMQTTCore, Warning, TEXT("Attempted to subscribe with an empty topic"));
+				return MakeFulfilledPromise<TArray<FMQTTSubscribeResult>>(FailureResponse).GetFuture();	
+			}
+
+			if(!FMQTTTopic::IsValid(TopicFilterQoSPair.Key))
+			{
+				UE_LOG(LogMQTTCore, Warning, TEXT("Attempted to subscribe with an invalid or malformed topic: \'%s\'"), *TopicFilterQoSPair.Key);
+				return MakeFulfilledPromise<TArray<FMQTTSubscribeResult>>(FailureResponse).GetFuture();
+			}
+			
+			TopicFilters.Add(TopicFilterQoSPair.Key);
+			TopicQoS.Add(TopicFilterQoSPair.Value);		
+		}
+
+		const uint16 PacketId = Connection->PopId();
+
+		UE_LOG(LogMQTTCore, Verbose, TEXT("Queued Subscribe message with PacketId %d., and Topic Filter: \'%s\'"), PacketId, *TopicFilters[0]);
+		
+		return Connection->QueueOperation<FMQTTSubscribeOperation>(
+			FMQTTSubscribePacket(
+				PacketId,
+				TopicFilters,
+				TopicQoS))
+		.Next([Self = TWeakPtr<FMQTTClient, ESPMode::ThreadSafe>(AsShared()), InTopicFilterQoSPairs](FMQTTSubscribeAckPacket InResponse)
+		{
+			TArray<FMQTTSubscribeResult> Result;
+			if(!Self.IsValid())
+			{
+				return Result;
+			}
+
+			if(auto SharedSelf = Self.Pin())
+			{
+				Result.Reserve(InResponse.ReturnCodes.Num());
+
+				// Was abandoned
+				if(InResponse.ReturnCodes.IsEmpty())
+				{
+					return Result;				
+				}
+			
+				ensure(InResponse.ReturnCodes.Num() == InTopicFilterQoSPairs.Num());
+				for(int32 Idx = 0; Idx < FMath::Min(InResponse.ReturnCodes.Num(), InTopicFilterQoSPairs.Num()); ++Idx)
+				{
+					const TSharedPtr<FMQTTSubscription, ESPMode::ThreadSafe> Subscription = SharedSelf->MakeSubscription(InTopicFilterQoSPairs[Idx].Key, static_cast<EMQTTQualityOfService>(InResponse.ReturnCodes[Idx]));
+					Result.Add({static_cast<EMQTTSubscribeReturnCode>(InResponse.ReturnCodes[Idx]), Subscription});
+				}
+
+				SharedSelf->OnSubscribe().Broadcast(Result);
+			}			
+
+			return Result;
+		});
+	}
+
+	return MakeFulfilledPromise<TArray<FMQTTSubscribeResult>>(FailureResponse).GetFuture();
+}
+
+TFuture<bool> FMQTTClient::Unsubscribe(const TSet<FString>& InTopicFilters)
+{
+	if(Connection.IsValid())
+	{
+		return Connection->QueueOperation<FMQTTUnsubscribeOperation>(
+			FMQTTUnsubscribePacket(
+				Connection->PopId(),
+				InTopicFilters.Array()))
+		.Next([Self = TWeakPtr<FMQTTClient, ESPMode::ThreadSafe>(AsShared())](FMQTTUnsubscribeAckPacket InResponse)
+		{
+			if(!Self.IsValid())
+			{
+				return false;
+			} 
+			
+			Self.Pin()->OnUnsubscribe().Broadcast();
+			return true;
+		});
+	}
+
+	return MakeFulfilledPromise<bool>(false).GetFuture();
+}
+
+TFuture<bool> FMQTTClient::Ping(const float& InTimeout)
+{
+	if(Connection.IsValid())
+	{
+		return Connection->QueueOperation<FMQTTPingOperation>({})
+		.Next([](FMQTTPingResponsePacket InResponse)
+		{
+			return true;
+		});
+	}
+
+	return MakeFulfilledPromise<bool>(false).GetFuture();
+}
+
+bool FMQTTClient::IsConnected() const
+{
+	return Connection.IsValid() && Connection->GetState() == EMQTTState::Connected;
+}
+
+bool FMQTTClient::IsValid() const
+{
+	return ClientId.IsValid() && URL.IsValid();
+}
+
+TSharedPtr<FMQTTSubscription, ESPMode::ThreadSafe>& FMQTTClient::MakeSubscription(const FString& InTopicFilter, EMQTTQualityOfService InGrantedQoS)
+{
+	FScopeLock Lock(&SubscriptionsLock);
+	return Subscriptions.Add(
+		FName(*InTopicFilter),
+		MakeShared<FMQTTSubscription, ESPMode::ThreadSafe>(InTopicFilter, InGrantedQoS));
+}
+
+// Handle raw incoming publish packet
+void FMQTTClient::OnMessagePacket(const FMQTTPublishPacket& InPacket)
+{
+	FMQTTClientMessage Message;
+	Message.TimeStamp = FDateTime::UtcNow(); // @todo: get sooner than this
+	Message.Topic = InPacket.Topic;
+	Message.Payload = InPacket.Payload;
+	Message.QoS = InPacket.QoS;
+	Message.bRetain = InPacket.bRetain;
+	Message.PayloadString = FString(Message.Payload.Num(), (const char*)Message.Payload.GetData());
+
+	// No subscriptions, can happen if that overload isn't used
+	if(Subscriptions.Num() == 0)
+	{
+		OnMessage().Broadcast(Message);
+	}
+	else
+	{
+		// Route message
+		Async(EAsyncExecution::TaskGraph, [&, Message]()
+		{		
+			TArray<TSharedPtr<FMQTTSubscription, ESPMode::ThreadSafe>> SubscriptionsCopy;
+			{
+				FScopeLock Lock(&SubscriptionsLock);
+				Subscriptions.GenerateValueArray(SubscriptionsCopy);
+			}
+
+			OnMessage().Broadcast(Message);
+
+			FString Topic = Message.Topic;
+			Algo::ForEachIf(
+				SubscriptionsCopy,
+				[&Topic](const auto& InSubscription) { return InSubscription->Matches(Topic); },
+				[&Message](const auto& InSubscription)
+				{
+					InSubscription->OnSubscriptionMessage().Broadcast(Message);
+				});
+		});
+	}
+}
diff --git a/Engine/Plugins/Protocols/MQTT/Source/MQTTCore/Private/MQTTClient.h b/Engine/Plugins/Protocols/MQTT/Source/MQTTCore/Private/MQTTClient.h
new file mode 100644
--- /dev/null
+++ b/Engine/Plugins/Protocols/MQTT/Source/MQTTCore/Private/MQTTClient.h
@@ -0,0 +1,61 @@
+﻿// Copyright Epic Games, Inc. All Rights Reserved.
+
+#pragma once
+
+#include "IMQTTClient.h"
+#include "MQTTConnection.h"
+
+class FMQTTClient final
+	: public IMQTTClient
+	, public TSharedFromThis<FMQTTClient, ESPMode::ThreadSafe>
+{
+public:
+	explicit FMQTTClient(const FMQTTURL& InURL);
+	virtual ~FMQTTClient() override;
+
+	//~ Begin IMQTTClient Interface
+	virtual FOnConnect& OnConnect() override { return OnConnectDelegate; }
+	virtual FOnDisconnect& OnDisconnect() override { return OnDisconnectDelegate; }
+	virtual FOnPublish& OnPublish() override { return OnPublishDelegate; }
+	virtual FOnMessage& OnMessage() override { return OnMessageDelegate; }
+	virtual FOnSubscribe& OnSubscribe() override { return OnSubscribeDelegate; }
+	virtual FOnUnsubscribe& OnUnsubscribe() override { return OnUnsubscribeDelegate; }
+
+	virtual FGuid GetClientId() const override { return ClientId; }
+
+	virtual const FMQTTURL& GetURL() const override	{ return URL; }	
+
+	virtual TFuture<EMQTTConnectReturnCode> Connect(bool bCleanSession = true) override;
+	virtual TFuture<void> Disconnect() override;
+
+	virtual TFuture<bool> Publish(const FString& InTopic, const TArray<uint8>& InPayload, EMQTTQualityOfService InQoS = EMQTTQualityOfService::Once, const bool bInRetain = false) override;
+	virtual TFuture<bool> Publish(const FString& InTopic, const FString& InPayload, EMQTTQualityOfService InQoS = EMQTTQualityOfService::Once, const bool bInRetain = false) override;
+	virtual TFuture<TArray<FMQTTSubscribeResult>> Subscribe(const TArray<TPair<FString, EMQTTQualityOfService>>& InTopicFilterQoSPairs) override;
+	virtual TFuture<bool> Unsubscribe(const TSet<FString>& InTopicFilters) override;
+	virtual TFuture<bool> Ping(const float& InTimeout = 2.0f) override;
+
+	virtual bool IsConnected() const override;
+
+	virtual bool IsValid() const override;
+	//~ End IMQTTClient Interface
+
+protected:
+	TSharedPtr<FMQTTSubscription, ESPMode::ThreadSafe>& MakeSubscription(const FString& InTopicFilter, EMQTTQualityOfService InGrantedQoS);
+	void OnMessagePacket(const FMQTTPublishPacket& InPacket);
+
+private:
+	TSharedPtr<FMQTTConnection, ESPMode::ThreadSafe> Connection;
+	
+	FOnConnect OnConnectDelegate;
+	FOnDisconnect OnDisconnectDelegate;	
+	FOnPublish OnPublishDelegate;	
+	FOnSubscribe OnSubscribeDelegate;
+	FOnUnsubscribe OnUnsubscribeDelegate;
+	FOnMessage OnMessageDelegate;
+
+	FGuid ClientId;
+	FMQTTURL URL;
+
+	mutable FCriticalSection SubscriptionsLock;
+	TMap<FName, TSharedPtr<FMQTTSubscription, ESPMode::ThreadSafe>> Subscriptions;
+};
diff --git a/Engine/Plugins/Protocols/MQTT/Source/MQTTCore/Private/MQTTClientMessage.cpp b/Engine/Plugins/Protocols/MQTT/Source/MQTTCore/Private/MQTTClientMessage.cpp
new file mode 100644
--- /dev/null
+++ b/Engine/Plugins/Protocols/MQTT/Source/MQTTCore/Private/MQTTClientMessage.cpp
@@ -0,0 +1,17 @@
+﻿// Copyright Epic Games, Inc. All Rights Reserved.
+
+#include "MQTTClientMessage.h"
+
+#include "Serialization/JsonReader.h"
+#include "Serialization/JsonSerializer.h"
+
+bool FMQTTClientMessage::GetPayloadAsJson(TSharedPtr<FJsonObject>& OutJson) const
+{
+	return FJsonSerializer::Deserialize(TJsonReaderFactory<TCHAR>::Create(PayloadString), OutJson);
+}
+
+void FMQTTClientMessage::SetPayloadFromString(const FString& InPayloadString)
+{
+	PayloadString = InPayloadString;
+	Payload = TArray<uint8>((uint8*)TCHAR_TO_UTF8(*InPayloadString), InPayloadString.Len());
+}
diff --git a/Engine/Plugins/Protocols/MQTT/Source/MQTTCore/Private/MQTTClientObject.cpp b/Engine/Plugins/Protocols/MQTT/Source/MQTTCore/Private/MQTTClientObject.cpp
new file mode 100644
--- /dev/null
+++ b/Engine/Plugins/Protocols/MQTT/Source/MQTTCore/Private/MQTTClientObject.cpp
@@ -0,0 +1,266 @@
+﻿// Copyright Epic Games, Inc. All Rights Reserved.
+
+#include "MQTTClientObject.h"
+
+#include "IMQTTCoreModule.h"
+#include "MQTTClient.h"
+#include "MQTTCoreLog.h"
+
+//~ Begin UMQTTSubscriptionObject
+UMQTTSubscriptionObject::~UMQTTSubscriptionObject()
+{
+	if (const TSharedPtr<FMQTTSubscription, ESPMode::ThreadSafe> MQTTSubscription = MqttSubscriptionPtr)
+	{
+		MQTTSubscription->OnSubscriptionMessage().RemoveAll(this);
+	}
+}
+
+void UMQTTSubscriptionObject::Initialize(const TSharedPtr<FMQTTSubscription, ESPMode::ThreadSafe>& InSubscription)
+{
+	check(InSubscription.IsValid());
+
+	MqttSubscriptionPtr = InSubscription;
+	MqttSubscriptionPtr->OnSubscriptionMessage().AddUObject(this, &UMQTTSubscriptionObject::OnMessage);
+}
+
+void UMQTTSubscriptionObject::SetOnMessageHandler(const FOnMessageDelegate& OnMessageCallback)
+{
+	OnMessageDelegate = OnMessageCallback;
+}
+
+void UMQTTSubscriptionObject::OnMessage(const FMQTTClientMessage& InMessage) const
+{
+	OnMessageDelegate.ExecuteIfBound(InMessage);
+}
+//~ End UMQTTSubscriptionObject
+
+UMQTTClientObject::~UMQTTClientObject()
+{
+	RemoveDelegates();
+	if(MqttClientPtr.IsValid())
+	{
+		MqttClientPtr.Reset();
+	}
+}
+
+void UMQTTClientObject::Initialize(const FMQTTURL& InURL)
+{
+	const TSharedPtr<IMQTTClient, ESPMode::ThreadSafe> MQTTClient = IMQTTCoreModule::Get().GetOrCreateClient(InURL);
+	if (MQTTClient)
+	{
+		MqttClientPtr = MQTTClient;
+		InitDelegates();
+	}
+}
+
+void UMQTTClientObject::InitDelegates()
+{
+	if (const TSharedPtr<IMQTTClient, ESPMode::ThreadSafe> MQTTClient = MqttClientPtr)
+	{
+		// MQTTClient->OnConnect().AddUObject(this, &UMQTTClientObject::OnConnect);
+		// MQTTClient->OnDisconnect().AddUObject(this, &UMQTTClientObject::OnDisconnect);
+		// MQTTClient->OnPublish().AddUObject(this, &UMQTTClientObject::OnPublish);
+		// MQTTClient->OnMessage().AddUObject(this, &UMQTTClientObject::OnMessage);
+		// MQTTClient->OnSubscribe().AddLambda([&](TArray<FMQTTSubscribeResult> InSubscriptions)
+		// {
+		// 	// Get list of granted quality of service levels, one per topic subscribed to (usually only one)
+		// 	TArray<EMQTTQualityOfService> QoS;
+		// 	Algo::Transform(
+		// 		InSubscriptions,
+		// 		QoS,
+		// 		[](const FMQTTSubscribeResult& InItem)
+		// 		{
+		// 			return InItem.Subscription->GetGrantedQoS();
+		// 		});
+		// 	
+		// 	OnSubscribe(QoS);
+		// });
+		// MQTTClient->OnUnsubscribe().AddUObject(this, &UMQTTClientObject::OnUnsubscribe);
+	}
+}
+
+void UMQTTClientObject::RemoveDelegates() const
+{
+	if (const TSharedPtr<IMQTTClient, ESPMode::ThreadSafe> MQTTClient = MqttClientPtr)
+	{
+		MQTTClient->OnConnect().RemoveAll(this);
+		MQTTClient->OnDisconnect().RemoveAll(this);
+		MQTTClient->OnPublish().RemoveAll(this);
+		MQTTClient->OnMessage().RemoveAll(this);
+		MQTTClient->OnSubscribe().RemoveAll(this);
+		MQTTClient->OnUnsubscribe().RemoveAll(this);
+	}
+}
+
+void UMQTTClientObject::BeginDestroy()
+{
+	if (const TSharedPtr<IMQTTClient, ESPMode::ThreadSafe> MQTTClient = MqttClientPtr)
+	{
+		//IMQTTCoreModule::Get().DestroyClient(MQTTClient);
+	}
+	
+	Super::BeginDestroy();
+}
+
+TSharedPtr<IMQTTClient, ESPMode::ThreadSafe> UMQTTClientObject::GetClientImpl() const
+{
+	return MqttClientPtr;
+}
+
+void UMQTTClientObject::Connect(const FOnConnectDelegate& InOnConnect)
+{
+	if (const TSharedPtr<IMQTTClient, ESPMode::ThreadSafe> MQTTClient = MqttClientPtr)
+	{
+		MQTTClient->Connect()
+		.Next([&](EMQTTConnectReturnCode InReturnCode)
+		{
+			if(IsValid(this))
+			{
+				InOnConnect.ExecuteIfBound(InReturnCode);
+			}
+		});
+	}
+}
+
+void UMQTTClientObject::Disconnect(const FOnDisconnectDelegate& InOnDisconnect)
+{
+	OnDisconnectDelegate = InOnDisconnect;
+	if (const TSharedPtr<IMQTTClient, ESPMode::ThreadSafe> MQTTClient = MqttClientPtr)
+	{
+		const FGuid ClientId = MQTTClient->GetClientId();
+		MQTTClient->Disconnect()
+		.Next([&](int32)
+		{
+			if(IsValid(this))
+			{
+				InOnDisconnect.ExecuteIfBound();
+			}
+		});
+	}
+}
+
+void UMQTTClientObject::Publish(
+	const FString& InTopic,
+	const TArray<uint8>& InPayload,
+	EMQTTQualityOfService InQoS,
+	const bool bInRetain)
+{
+	if (const TSharedPtr<IMQTTClient, ESPMode::ThreadSafe> MQTTClient = MqttClientPtr)
+	{
+		MQTTClient->Publish(InTopic, InPayload, InQoS, bInRetain);
+	}
+}
+
+UMQTTSubscriptionObject* UMQTTClientObject::Subscribe(const FString& InTopic, EMQTTQualityOfService InQoS)
+{
+	if (const TSharedPtr<IMQTTClient, ESPMode::ThreadSafe> MQTTClient = MqttClientPtr)
+	{
+		UMQTTSubscriptionObject* SubscriptionObject = NewObject<UMQTTSubscriptionObject>(this);
+		UMQTTSubscriptionObject*& SubscriptionObjectRef = this->Subscriptions.Add_GetRef(SubscriptionObject);
+
+		MQTTClient->Subscribe({MakeTuple(InTopic, InQoS)})
+		.Next([&](TArray<FMQTTSubscribeResult> InSubscribeResults)
+		{
+			if(IsValid(this))
+			{
+				if(InSubscribeResults.IsEmpty() || InSubscribeResults[0].ReturnCode == EMQTTSubscribeReturnCode::Failure)
+				{
+					return;
+				}
+
+				TSharedPtr<FMQTTSubscription, ESPMode::ThreadSafe> Subscription = InSubscribeResults[0].Subscription;
+				SubscriptionObjectRef->Initialize(Subscription);
+				// Subscription->OnSubscriptionMessage() // @todo: get from user callbacks?
+			}
+		});
+
+		return SubscriptionObject;
+	}
+
+	return nullptr;
+}
+
+TArray<UMQTTSubscriptionObject*> UMQTTClientObject::SubscribeMany(
+	const TArray<FString>& InTopics,
+	const TArray<EMQTTQualityOfService> InQoS)
+{
+	if (const TSharedPtr<IMQTTClient, ESPMode::ThreadSafe> MQTTClient = MqttClientPtr)
+	{
+		ensure(InTopics.Num() == InQoS.Num());
+		
+		// @todo: how to prevent same topic returning duplicate subscription
+		TArray<UMQTTSubscriptionObject*> Result;
+		Result.Reserve(InTopics.Num());
+
+		TArray<TPair<FString, EMQTTQualityOfService>> TopicQoSPairs;
+		TopicQoSPairs.Reserve(InTopics.Num());
+		
+		for(auto Idx = 0; Idx < InTopics.Num(); ++Idx)
+		{
+			UMQTTSubscriptionObject* SubscriptionObject = NewObject<UMQTTSubscriptionObject>(this);
+			Subscriptions.Add(SubscriptionObject);
+			Result[Idx] = SubscriptionObject;
+
+			TopicQoSPairs[Idx] = MakeTuple(InTopics[Idx], InQoS[Idx]);
+		}
+
+		MQTTClient->Subscribe(TopicQoSPairs)
+		.Next([&](TArray<FMQTTSubscribeResult> InSubscribeResults)
+		{
+			if(IsValid(this))
+			{
+				if(InSubscribeResults.IsEmpty() || InSubscribeResults[0].ReturnCode == EMQTTSubscribeReturnCode::Failure)
+				{
+					return;
+				}
+
+				for(auto Idx = 0; Idx < InSubscribeResults.Num(); ++Idx)
+				{
+					TSharedPtr<FMQTTSubscription, ESPMode::ThreadSafe> Subscription = InSubscribeResults[Idx].Subscription;
+					auto* SubscriptionObject = Subscriptions.FindByPredicate([&](const UMQTTSubscriptionObject* InSubscription)
+					{
+						return IsValid(InSubscription) && GetTypeHash(InSubscription) == GetTypeHash(Subscription.Get());
+					});
+					
+					if(!SubscriptionObject)
+					{
+						UE_LOG(LogMQTTCore, Error, TEXT("Failed to find SubscriptionObject for topic: '%s'"), *Subscription->GetTopic().ToString());
+						continue;
+					}
+					
+					(*SubscriptionObject)->Initialize(Subscription);
+				}
+			}
+		});
+	}
+
+	return {};
+}
+
+void UMQTTClientObject::Unsubscribe(const FString& InTopic)
+{
+	if (const TSharedPtr<IMQTTClient, ESPMode::ThreadSafe> MQTTClient = MqttClientPtr)
+	{
+		MQTTClient->Unsubscribe({InTopic});
+	}
+}
+
+FString UMQTTClientObject::GetClientId() const
+{
+	if (const TSharedPtr<IMQTTClient, ESPMode::ThreadSafe> MQTTClient = MqttClientPtr)
+	{
+		return MQTTClient->GetClientId().ToString();
+	}
+
+	return FString();
+}
+
+FMQTTURL UMQTTClientObject::GetURL() const
+{
+	if (const TSharedPtr<IMQTTClient, ESPMode::ThreadSafe> MQTTClient = MqttClientPtr)
+	{
+		return MQTTClient->GetURL();
+	}
+
+	return {};
+}
diff --git a/Engine/Plugins/Protocols/MQTT/Source/MQTTCore/Private/MQTTConnection.cpp b/Engine/Plugins/Protocols/MQTT/Source/MQTTCore/Private/MQTTConnection.cpp
new file mode 100644
--- /dev/null
+++ b/Engine/Plugins/Protocols/MQTT/Source/MQTTCore/Private/MQTTConnection.cpp
@@ -0,0 +1,790 @@
+﻿// Copyright Epic Games, Inc. All Rights Reserved.
+
+#include "MQTTConnection.h"
+
+#include "IMQTTCoreModule.h"
+#include "MQTTClient.h"
+#include "MQTTClientSettings.h"
+#include "MQTTCoreLog.h"
+#include "MQTTCoreStats.h"
+#include "MQTTSharedInternal.h"
+#include "Sockets.h"
+#include "SocketSubsystem.h"
+#include "Async/Async.h"
+#include "HAL/Runnable.h"
+#include "HAL/RunnableThread.h"
+#include "Math/UnitConversion.h"
+#include "Serialization/ArrayReader.h"
+#include "Serialization/BufferArchive.h"
+#include "Templates/Atomic.h"
+
+DECLARE_DWORD_ACCUMULATOR_STAT(TEXT("MQTT Packets Sent Total"), STAT_MQTTPackagesSent, STATGROUP_MQTT);
+DECLARE_DWORD_ACCUMULATOR_STAT(TEXT("MQTT Packets Received Total"), STAT_MQTTPackagesReceived, STATGROUP_MQTT);
+DECLARE_FLOAT_COUNTER_STAT(TEXT("MQTT Execution Delay Average"), STAT_MQTTExecutionTime, STATGROUP_MQTT);
+
+#define LOCTEXT_NAMESPACE "MQTTConnection"
+
+FMQTTPacketIdGenerator::FMQTTPacketIdGenerator()
+	: Min(1)
+	, Position(1)
+{
+	for(auto Idx = 0; Idx < ValueStore.Num(); ++Idx)
+	{
+		ValueStore[Idx] = Idx;
+	}	
+}
+
+uint16 FMQTTPacketIdGenerator::Pop()
+{
+	// Shift so it's in single range (shifted pos -> max)
+	const uint16 ShiftedPosition = (Position++) - Min;
+	// Then shift back to origin, and wrap
+	return ValueStore[(ShiftedPosition + Min) % Max];		
+}
+
+void FMQTTPacketIdGenerator::Push(const uint16 InValue)
+{
+	ValueStore[Min++] = InValue;
+}
+
+FMQTTConnection::FMQTTConnection(
+	FSocket& InSocket,
+	const TSharedRef<FInternetAddr> InAddr)
+	: Socket(&InSocket)
+	, EndpointInternetAddr(InAddr)
+	, Thread(nullptr)
+{
+	check(Socket->GetSocketType() == SOCKTYPE_Streaming);
+
+	ThreadName = FString(TEXT("MQTTConnection_")) + InAddr->ToString(true);
+	Thread = FRunnableThread::Create(this, *ThreadName, 0, TPri_TimeCritical, FPlatformAffinity::GetPoolThreadMask());
+
+	UE_LOG(LogMQTTCore, VeryVerbose, TEXT("Created MQTTConnection for %s"), *EndpointInternetAddr->ToString(false));
+}
+
+FMQTTConnection::~FMQTTConnection()
+{
+	SetState(EMQTTState::Stopping); 
+
+	AbandonOperations();
+
+	if (Socket)
+	{
+		ISocketSubsystem* SocketSubsystem = ISocketSubsystem::Get(PLATFORM_SOCKETSUBSYSTEM);
+		SocketSubsystem->DestroySocket(Socket);
+	}
+	
+	if (Thread != nullptr)
+	{
+		Thread->Kill(true);
+		Thread = nullptr;
+		delete Thread;
+	}
+
+	if (Socket)
+	{
+		Socket = nullptr;
+	}
+
+	UE_LOG(LogMQTTCore, VeryVerbose, TEXT("Destroyed MQTTConnection at %s"), *EndpointInternetAddr->ToString(false));
+}
+
+TSharedPtr<FMQTTConnection, ESPMode::ThreadSafe> FMQTTConnection::TryCreate(const TSharedPtr<FInternetAddr>& InAddr)
+{
+	if(!InAddr.IsValid() || !InAddr->IsValid())
+	{
+		UE_LOG(LogMQTTCore, Error, TEXT("Cannot create MQTTConnection, the IP Address was invalid."));
+		return nullptr;		
+	}
+	
+	ISocketSubsystem* SocketSubsystem = ISocketSubsystem::Get(PLATFORM_SOCKETSUBSYSTEM);
+	FSocket* Socket = SocketSubsystem->CreateSocket(NAME_Stream, TEXT("MQTTConnection"), FNetworkProtocolTypes::IPv4);
+	Socket->SetReuseAddr(true);
+	check(Socket);
+
+	if (!Socket)
+	{
+		UE_LOG(LogMQTTCore, Error, TEXT("Cannot create MQTTConnection: Error creating Socket for: %s: %s"), *InAddr->ToString(true), SocketSubsystem->GetSocketError(SocketSubsystem->GetLastErrorCode()));
+		return nullptr;
+	}
+
+	TSharedPtr<FMQTTConnection, ESPMode::ThreadSafe> Runnable = MakeShared<FMQTTConnection, ESPMode::ThreadSafe>(*Socket, InAddr.ToSharedRef());
+	return Runnable;
+}
+
+bool FMQTTConnection::Init()
+{
+	return FRunnable::Init();
+}
+
+uint32 FMQTTConnection::Run()
+{
+	const UMQTTClientSettings* MQTTSettings = GetDefault<UMQTTClientSettings>();
+	check(MQTTSettings);
+
+	// Fixed rate delta time
+	const double SendDeltaTime = 1.0f / MQTTSettings->PublishRate;
+
+	while (GetState() != EMQTTState::Stopping)
+	{
+		const double StartTime = FPlatformTime::Seconds();
+
+		if (UpdateConnection())
+		{
+			ProcessMessages();
+		}
+
+		const double EndTime = FPlatformTime::Seconds();
+		const double WaitTime = SendDeltaTime - (EndTime - StartTime);
+
+		if (WaitTime > 0.f)
+		{
+			// Sleep by the amount which is set in refresh rate
+			FPlatformProcess::SleepNoStats(WaitTime);
+		}
+	}
+
+	return 0;
+}
+
+void FMQTTConnection::Stop()
+{
+	AbandonOperations();
+	SetState(EMQTTState::Stopping);
+}
+
+void FMQTTConnection::Exit()
+{
+	FRunnable::Exit();	
+}
+
+void FMQTTConnection::Tick()
+{
+	ProcessMessages();
+}
+
+FSingleThreadRunnable* FMQTTConnection::GetSingleThreadInterface()
+{
+	return this;
+}
+
+bool FMQTTConnection::UpdateConnection()
+{
+	EMQTTSocketState CurrentSocketState = SocketState;
+	if(CurrentSocketState == EMQTTSocketState::Connected)
+	{
+		ISocketSubsystem* SocketSubsystem = ISocketSubsystem::Get(PLATFORM_SOCKETSUBSYSTEM);
+		const ESocketErrors LastErrorCode = SocketSubsystem->GetLastErrorCode();
+		if(LastErrorCode > ESocketErrors::SE_EINPROGRESS)
+		{
+			CurrentSocketState = EMQTTSocketState::Error;	
+		}
+	}
+
+	switch(CurrentSocketState)
+	{
+	// Already connected, skip
+	case EMQTTSocketState::Connected:
+		break;
+
+	case EMQTTSocketState::ConnectionRejected:
+	case EMQTTSocketState::Stopping:
+		{
+			// Unless connecting, abandon queued operations
+			if(CurrentSocketState != EMQTTSocketState::Connecting)
+			{
+				if(PendingConnectOperation.IsValid())
+				{
+					PendingConnectOperation->Complete(FMQTTConnectAckPacket(EMQTTConnectReturnCode::SocketError));
+				}
+			
+				AbandonOperations();
+			}
+		
+			return false;
+		}
+
+	case EMQTTSocketState::Connecting:
+	case EMQTTSocketState::Error:
+	default:
+		{
+			const ESocketConnectionState ConnectionState = Socket->GetConnectionState();
+			if(ConnectionState == ESocketConnectionState::SCS_Connected)
+			{
+				SocketState = EMQTTSocketState::Connected;
+			}
+			// Socket isn't connected
+			else
+			{
+				Socket->SetNonBlocking(true);
+				OpenSocket();
+				Socket->SetNonBlocking(false);
+
+				// If there's an error, abandon all operations until next explicit connect call
+				if(SocketState == EMQTTSocketState::Error)
+				{
+					if(PendingConnectOperation.IsValid())
+					{
+						PendingConnectOperation->Complete(FMQTTConnectAckPacket(EMQTTConnectReturnCode::SocketError));
+					}
+			
+					AbandonOperations();
+				}
+
+				return false;
+			}
+		}
+	}
+
+	// Check MQTT Connect timeout
+	if(GetState() == EMQTTState::Connecting && PendingConnectOperation.IsValid())
+	{
+		if(FPlatformTime::Seconds() > LastConnectRequestTime + MQTTConnectionTimeout)
+		{
+			PendingConnectOperation->Complete(FMQTTConnectAckPacket(EMQTTConnectReturnCode::SocketError));
+			PendingConnectOperation.Reset();
+			SetState(EMQTTState::Disconnected);
+			return false;
+		}
+	}
+
+	return true;
+}
+
+void FMQTTConnection::ProcessMessages()
+{
+	UE_LOG(LogMQTTCore, Verbose, TEXT("BEEP: %s"), MQTT::GetMQTTStateName(GetState()));
+
+	if(GetState() == EMQTTState::Stopping)
+	{
+		UE_LOG(LogMQTTCore, Verbose, TEXT("BEEP was stopped."));
+		return;
+	}
+	
+	// KeepAlive ping
+	if(FPlatformTime::Seconds() >= LastKeepAliveTime)
+	{
+		QueueOperation<FMQTTPingOperation>({})
+		.Next([](FMQTTPingResponsePacket InResponse)
+		{
+			// @todo: something if no response
+			return true;
+		});
+		
+		LastKeepAliveTime = FPlatformTime::Seconds() + KeepAliveInterval;
+	}
+
+	TQueue<TUniquePtr<IMQTTOperation>, EQueueMode::Spsc> PendingOutgoingOperations;
+	{
+		UE_LOG(LogMQTTCore, Verbose, TEXT("Copy outgoing operations to buffer"));
+		
+		uint32 DeferredOperationCount = 0;
+		
+		// Copy to temporary buffer
+		FScopeLock Lock(&OutgoingQueueLock);
+		TUniquePtr<IMQTTOperation> Operation;
+		TQueue<TUniquePtr<IMQTTOperation>> DeferredOperations;
+		while(OutgoingOperations.Dequeue(Operation))
+		{
+			// Was set in previous tick, skip
+			if(!Operation.IsValid() || !Operation)
+			{
+				continue;
+			}
+			
+			if(Operation->StateDependency == EMQTTState::None || Operation->StateDependency == GetState())
+			{
+				PendingOutgoingOperations.Enqueue(MoveTemp(Operation));
+			}
+			else
+			{
+				DeferredOperationCount++;
+				DeferredOperations.Enqueue(MoveTemp(Operation));
+			}
+		}
+
+		UE_LOG(LogMQTTCore, Verbose, TEXT("Operations deferred: %d"), DeferredOperationCount);
+
+		// Move deferred back into queue for future processing
+		while(DeferredOperations.Dequeue(Operation))
+		{
+			OutgoingOperations.Enqueue(MoveTemp(Operation));
+		}
+	}
+
+	TQueue<TUniquePtr<IMQTTOperation>, EQueueMode::Spsc> PendingIncoming;
+
+	// Send outgoing packets
+	TUniquePtr<IMQTTOperation> Operation;
+	while(PendingOutgoingOperations.Dequeue(Operation))
+	{
+#if WITH_EDITORONLY_DATA
+		Operation->MarkExecuted();
+		SET_FLOAT_STAT(STAT_MQTTExecutionTime, Operation->GetTimeToExecution());
+#endif
+		
+		FBufferArchive BufferArchive;
+		Operation->Pack(BufferArchive);
+		
+		int32 ExpectedBytesSent = BufferArchive.Num();
+		int32 ActualBytesSent = -1;
+
+ 		if(Socket->SendTo(BufferArchive.GetData(), BufferArchive.Num(), ActualBytesSent, *EndpointInternetAddr))
+		{
+			INC_DWORD_STAT(STAT_MQTTPackagesSent);			
+		}
+		else
+		{
+			ISocketSubsystem* SocketSubsystem = ISocketSubsystem::Get(PLATFORM_SOCKETSUBSYSTEM);
+			UE_LOG(LogMQTTCore, Error, TEXT("Error sending MQTT packet: %s: %s"), *EndpointInternetAddr->ToString(true), SocketSubsystem->GetSocketError(SocketSubsystem->GetLastErrorCode()));
+		}
+
+		if(ActualBytesSent != ExpectedBytesSent)
+		{
+			UE_LOG(LogMQTTCore, Error, TEXT("Packet size mismatch: expected %d, was %d"), ExpectedBytesSent, ActualBytesSent);
+		}
+
+		if(Operation->GetPendingPacketType() != EMQTTPacketType::None)
+		{
+			PendingIncoming.Enqueue(MoveTemp(Operation));
+		}
+		// Packet doesn't have a response, so set promise to indicate it's been processed 
+		else
+		{
+			uint16 ReturnedMessageId = Operation->Complete();
+			PushId(MoveTemp(ReturnedMessageId));
+		}
+	}
+
+	// Set outgoing operations with expected requests to PendingIncoming
+	if(!PendingIncoming.IsEmpty())
+	{
+		FScopeLock IncomingLock(&IncomingQueueLock);
+		while(PendingIncoming.Dequeue(Operation))
+		{
+#if UE_BUILD_DEBUG
+			UE_LOG(LogMQTTCore, Verbose, TEXT("Re-queue pending incoming operation: %s"), ANSI_TO_TCHAR(typeid(decltype(Operation)).name()));
+#endif
+			const EMQTTPacketType MessageType = Operation->GetPacketType();
+			if(MessageType == EMQTTPacketType::Connect)
+			{
+				// There should only be one Connect request at a time!
+				if(PendingConnectOperation.IsValid())
+				{
+					checkNoEntry();
+				}
+				
+				// @note: connection request made, so record timestamp
+				LastConnectRequestTime = FPlatformTime::Seconds();				
+				
+				PendingConnectOperation = TUniquePtr<FMQTTConnectOperation>(StaticCast<FMQTTConnectOperation*>(Operation.Release()));
+			}
+			else if(MessageType == EMQTTPacketType::PingRequest)
+			{
+				// If already waiting, cancel
+				if(PendingPingOperation.IsValid())
+				{
+					PendingPingOperation->Abandon();
+					PendingPingOperation.Reset();
+				}
+				PendingPingOperation = MoveTemp(Operation);
+			}
+			else
+			{
+				uint16 PacketId = Operation->GetPacketId();
+				PendingIncomingOperations.Emplace(PacketId, MoveTemp(Operation));
+			}
+		}
+	}
+
+	// Wait until something to do
+	if (!Socket->Wait(ESocketWaitConditions::WaitForRead, FTimespan::FromSeconds(0.1f)))
+	{
+		return;
+	}
+
+	// Process any incoming packets
+	uint32 Size = 0;
+	while (GetState() != EMQTTState::Stopping && Socket && Socket->HasPendingData(Size))
+	{
+		UE_LOG(LogMQTTCore, Verbose, TEXT("Processing incoming packets of size: %d"), Size);
+		
+		TSharedRef<FArrayReader> Reader = MakeShared<FArrayReader>(true);
+		Reader->SetNumUninitialized(Size);
+
+		int32 Read = 0;
+		if (Socket->RecvFrom(Reader->GetData(), Reader->Num(), Read, *EndpointInternetAddr))
+		{
+			// ensure((uint32)Read <= Size);
+			// Reader->RemoveAt(Read, Reader->Num() - Read, false);
+
+			int64 ReaderPos = Reader->Tell();
+			while(ReaderPos < Size)
+			{
+				FMQTTFixedHeader Header;
+				*Reader << Header;
+				Reader->Seek(ReaderPos); // go back to start of packet
+			
+				HandleMessage(Header.GetPacketType(), Reader);
+
+				ReaderPos = Reader->Tell();
+				
+				INC_DWORD_STAT(STAT_MQTTPackagesReceived);
+			}
+		}
+	}
+
+	UE_LOG(LogMQTTCore, Verbose, TEXT("BOOP"));
+}
+
+void FMQTTConnection::OpenSocket()
+{
+	if(Socket)
+	{
+		SocketState = EMQTTSocketState::Connecting;
+		
+		const double TimeoutTime = FPlatformTime::Seconds() + SocketConnectionTimeout;
+		do
+		{
+			if(Socket->Connect(*EndpointInternetAddr))
+			{
+				const ESocketConnectionState SocketConnectionState = Socket->GetConnectionState();
+				SocketState = SocketConnectionState == SCS_NotConnected
+								? EMQTTSocketState::Connecting
+								: static_cast<EMQTTSocketState>(SocketConnectionState); // ESocketConnectionState maps to EMQTTSocketState
+				return;
+			}
+			FPlatformProcess::SleepNoStats(0.5f);
+ 		}
+		while(FPlatformTime::Seconds() < TimeoutTime && GetState() != EMQTTState::Stopping);
+
+		if(SocketState != EMQTTSocketState::Stopping)
+		{
+			SocketState = EMQTTSocketState::Error;
+			ISocketSubsystem* SocketSubsystem = ISocketSubsystem::Get(PLATFORM_SOCKETSUBSYSTEM);
+			const ESocketErrors LastErrorCode = SocketSubsystem->GetLastErrorCode();
+			UE_LOG(LogMQTTCore, Error, TEXT("Error connecting socket to: %s: %s, Connection state: %s"),
+				*EndpointInternetAddr->ToString(true),
+				SocketSubsystem->GetSocketError(LastErrorCode),
+				MQTT::GetSocketConnectionStateName(Socket->GetConnectionState()));
+
+			if(LastErrorCode == SE_ECONNREFUSED
+				|| LastErrorCode == SE_EINVAL
+				|| LastErrorCode == SE_ETIMEDOUT)
+			{
+				// Don't attempt again unless explicit Connect call
+				SocketState = EMQTTSocketState::ConnectionRejected;
+
+				// If connection pending, cancel
+				if(PendingConnectOperation.IsValid())
+				{
+					PendingConnectOperation->Complete(FMQTTConnectAckPacket(EMQTTConnectReturnCode::SocketError));
+				}
+
+				AbandonOperations();
+			}
+		}
+	}
+}
+
+void FMQTTConnection::AbandonOperations()
+{
+	UE_LOG(LogMQTTCore, Verbose, TEXT("Abandoning Operations"));
+	
+	{
+		FScopeLock Lock(&OutgoingQueueLock);
+		TUniquePtr<IMQTTOperation> Operation;
+		while(OutgoingOperations.Dequeue(Operation))
+		{
+			Operation->Abandon();
+		}
+		OutgoingOperations.Empty();
+	}
+
+	{
+		FScopeLock Lock(&IncomingQueueLock);
+		TArray<TUniquePtr<IMQTTOperation>> Messages;
+		//PendingIncomingOperations.GenerateValueArray(Messages);
+		for(const TUniquePtr<IMQTTOperation>& Message : Messages)
+		{
+			Message->Abandon();
+		}
+		PendingIncomingOperations.Empty();
+	}
+
+	if(PendingConnectOperation.IsValid())
+	{
+		PendingConnectOperation->Abandon();
+		PendingConnectOperation.Reset();
+	}
+
+	if(PendingPingOperation.IsValid())
+	{
+		PendingPingOperation->Abandon();
+		PendingPingOperation.Reset();
+	}
+}
+
+void FMQTTConnection::HandleMessage(const EMQTTPacketType InMessageType, const TSharedRef<FArrayReader>& InPacketReader)
+{
+	switch(InMessageType)
+	{
+	case EMQTTPacketType::ConnectAcknowledge:
+		HandleConnectAcknowledge(InPacketReader);
+		break;
+
+	case EMQTTPacketType::SubscribeAcknowledge: 
+		HandleSubscribeAcknowledge(InPacketReader);
+		break;
+	
+	case EMQTTPacketType::UnsubscribeAcknowledge:
+		HandleUnsubscribeAcknowledge(InPacketReader);
+		break;
+
+	// Incoming published message
+	case EMQTTPacketType::Publish:
+		HandlePublish(InPacketReader);
+		break;
+
+	// Publish QoS 1
+	case EMQTTPacketType::PublishAcknowledge:
+		HandlePublishAcknowledge(InPacketReader);
+		break;
+		
+	// Publish QoS 2: Step 2
+	case EMQTTPacketType::PublishReceived:
+		HandlePublishReceived(InPacketReader);
+		break;
+
+	// Publish QoS 2: Step 3
+	case EMQTTPacketType::PublishRelease:
+		HandlePublishRelease(InPacketReader);
+		break;
+
+	// Publish QoS 2: Step 4
+	case EMQTTPacketType::PublishComplete:
+		HandlePublishComplete(InPacketReader);
+		break;
+
+	case EMQTTPacketType::PingResponse:
+		HandlePingResponse(InPacketReader);
+		break;
+		
+	case EMQTTPacketType::None:
+		{
+			UE_LOG(LogMQTTCore, Error, TEXT("PacketType was None"));
+		}
+		break;
+		
+	default:
+		{
+			UE_LOG(LogMQTTCore, Error, TEXT("Unhandled, unknown message type: %d (%s)"), InMessageType, MQTT::GetMQTTPacketTypeName(InMessageType));
+		}
+	}
+}
+
+void FMQTTConnection::HandleConnectAcknowledge(const TSharedRef<FArrayReader>& InPacketReader)
+{
+	FMQTTConnectAckPacket Message;
+	*InPacketReader << Message;
+	{
+		FScopeLock IncomingLock(&IncomingQueueLock);
+		TUniquePtr<IMQTTOperation> Operation = MoveTemp(PendingConnectOperation);
+		if(const auto CastMessageTask = StaticCast<TMQTTOperation<FMQTTConnectPacket, FMQTTConnectAckPacket>*>(Operation.Release()))
+		{
+			CastMessageTask->Complete(MoveTemp(Message));
+			PendingConnectOperation.Reset();
+		}
+	}
+	UE_LOG(LogMQTTCore, VeryVerbose, TEXT("Handled ConnectAck message."));
+}
+
+void FMQTTConnection::HandlePublishAcknowledge(const TSharedRef<FArrayReader>& InPacketReader)
+{
+	FMQTTPublishAckPacket Message;
+	*InPacketReader << Message;
+	{
+		FScopeLock IncomingLock(&IncomingQueueLock);
+		if(const TUniquePtr<IMQTTOperation>* Operation = PendingIncomingOperations.Find(Message.PacketId))
+		{
+			const uint16 PacketId = Message.PacketId;
+			if(const auto CastMessageTask = StaticCast<TMQTTOperation<FMQTTPublishPacket, FMQTTPublishAckPacket>*>(Operation->Get()))
+			{
+				CastMessageTask->Complete(MoveTemp(Message));
+			}
+			PendingIncomingOperations.Remove(PacketId);
+			PushId(PacketId);
+		}
+	}
+	UE_LOG(LogMQTTCore, VeryVerbose, TEXT("Handled PublishAck message."));
+}
+
+void FMQTTConnection::HandlePublishReceived(const TSharedRef<FArrayReader>& InPacketReader)
+{
+	FMQTTPublishReceivedPacket Message;
+	*InPacketReader << Message;
+	{
+		FScopeLock IncomingLock(&IncomingQueueLock);
+		const uint16 PacketId = Message.PacketId;
+		if(PendingIncomingOperations.Contains(PacketId))
+		{
+			TUniquePtr<IMQTTOperation> Operation = PendingIncomingOperations.FindAndRemoveChecked(PacketId);
+			if(const auto& CastMessageTask = StaticCast<FMQTTPublishOperationQoS2*>(Operation.Get()))
+			{
+				CastMessageTask->CompleteStep(MoveTemp(Message));
+				// Move from incoming to outgoing operations (PUBREL)
+				{
+					FScopeLock Lock(&OutgoingQueueLock);
+					OutgoingOperations.Enqueue(MoveTemp(Operation));
+				}
+				PendingIncomingOperations.Remove(PacketId);
+			}
+		}
+	}
+	UE_LOG(LogMQTTCore, VeryVerbose, TEXT("Handled PublishReceive message."));
+}
+
+void FMQTTConnection::HandlePublishRelease(const TSharedRef<FArrayReader>& InPacketReader)
+{
+	FMQTTPublishReleasePacket Message;
+	*InPacketReader << Message;
+	{
+		FScopeLock IncomingLock(&IncomingQueueLock);
+		const uint16 PacketId = Message.PacketId;
+		if(PendingIncomingOperations.Contains(PacketId))
+		{
+			TUniquePtr<IMQTTOperation> Operation = PendingIncomingOperations.FindAndRemoveChecked(PacketId);			
+			if(const auto& CastMessageTask = StaticCast<FMQTTPublishOperationQoS2*>(Operation.Get()))
+			{
+				CastMessageTask->CompleteStep(MoveTemp(Message));
+				// Move from incoming to outgoing operations (PUBCOMP)
+				{
+					FScopeLock Lock(&OutgoingQueueLock);
+					OutgoingOperations.Enqueue(MoveTemp(Operation));
+				}
+				PendingIncomingOperations.Remove(PacketId);
+			}
+		}
+	}
+	UE_LOG(LogMQTTCore, VeryVerbose, TEXT("Handled PublishRelease message."));
+}
+
+void FMQTTConnection::HandlePublishComplete(const TSharedRef<FArrayReader>& InPacketReader)
+{
+	FMQTTPublishCompletePacket Message;
+	*InPacketReader << Message;
+	{
+		FScopeLock IncomingLock(&IncomingQueueLock);
+		if(const TUniquePtr<IMQTTOperation>* Operation = PendingIncomingOperations.Find(Message.PacketId))
+		{
+			const uint16 PacketId = Message.PacketId;
+			if(const auto CastMessageTask = StaticCast<TMQTTOperation<FMQTTPublishReleasePacket, FMQTTPublishCompletePacket>*>(Operation->Get()))
+			{
+				CastMessageTask->Complete(MoveTemp(Message));
+			}
+			PendingIncomingOperations.Remove(PacketId);
+			PushId(PacketId);
+		}
+	}
+	UE_LOG(LogMQTTCore, VeryVerbose, TEXT("Handled PublishComplete message."));
+}
+
+void FMQTTConnection::HandleSubscribeAcknowledge(const TSharedRef<FArrayReader>& InPacketReader)
+{
+	FMQTTSubscribeAckPacket Message;
+	*InPacketReader << Message;
+	{
+		FScopeLock IncomingLock(&IncomingQueueLock);		
+		if(const TUniquePtr<IMQTTOperation>* Operation = PendingIncomingOperations.Find(Message.PacketId))
+		{
+			if(const auto CastMessageTask = StaticCast<TMQTTOperation<FMQTTSubscribePacket, FMQTTSubscribeAckPacket>*>(Operation->Get()))
+			{
+				CastMessageTask->Complete(MoveTemp(Message));
+			}
+			PendingIncomingOperations.Remove(Message.PacketId);
+			PushId(Message.PacketId);
+		}
+	}
+	UE_LOG(LogMQTTCore, VeryVerbose, TEXT("Handled SubscribeAck message with PacketId %d."), Message.PacketId);
+}
+
+void FMQTTConnection::HandleUnsubscribeAcknowledge(const TSharedRef<FArrayReader>& InPacketReader)
+{
+	FMQTTUnsubscribeAckPacket Message;
+	*InPacketReader << Message;
+	{
+		FScopeLock IncomingLock(&IncomingQueueLock);
+		if(const TUniquePtr<IMQTTOperation>* Operation = PendingIncomingOperations.Find(Message.PacketId))
+		{	
+			if(const auto CastMessageTask = StaticCast<TMQTTOperation<FMQTTUnsubscribePacket, FMQTTUnsubscribeAckPacket>*>(Operation->Get()))
+			{
+				CastMessageTask->Complete(MoveTemp(Message));
+			}
+			PendingIncomingOperations.Remove(Message.PacketId);
+			PushId(Message.PacketId);
+		}
+	}
+	UE_LOG(LogMQTTCore, VeryVerbose, TEXT("Handled UnsubscribeAck message."));
+}
+
+void FMQTTConnection::HandlePingResponse(const TSharedRef<FArrayReader>& InPacketReader)
+{
+	FMQTTPingResponsePacket Message;
+	*InPacketReader << Message;
+	{
+		FScopeLock IncomingLock(&IncomingQueueLock);
+		const TUniquePtr<IMQTTOperation>& Operation = PendingPingOperation;
+		if(const auto CastMessageTask = StaticCast<TMQTTOperation<FMQTTPingRequestPacket, FMQTTPingResponsePacket>*>(Operation.Get()))
+		{
+			CastMessageTask->Complete(MoveTemp(Message));
+		}
+	}
+	UE_LOG(LogMQTTCore, VeryVerbose, TEXT("Handled PingResponse message."));
+}
+
+void FMQTTConnection::HandlePublish(const TSharedRef<FArrayReader>& InPacketReader)
+{
+	FMQTTPublishPacket Message;
+	*InPacketReader << Message;
+	{
+		// @todo: if qos > 0, setup response chain
+		OnMessage().ExecuteIfBound(Message);
+		
+		// FScopeLock IncomingLock(&IncomingQueueLock);
+		// const TSharedPtr<IMQTTOperation, ESPMode::ThreadSafe> Operation = PendingPingOperation;
+		// if(const auto CastMessageTask = StaticCastSharedPtr<TMQTTOperation<FMQTTPublishPacket, FMQTTPingResponsePacket>>(Operation))
+		// {
+		// 	CastMessageTask->Complete(MoveTemp(Message));
+		// }
+	}
+	UE_LOG(LogMQTTCore, VeryVerbose, TEXT("Handled Publish message."));
+}
+
+void FMQTTConnection::SetState(const EMQTTState InState)
+{
+	// Can't exit stopping state
+	if(GetState() == EMQTTState::Stopping)
+	{
+		return;
+	}
+	
+	UE_LOG(LogMQTTCore, Verbose, TEXT("Set State to: %s"), MQTT::GetMQTTStateName(InState));
+	MQTTState.store(InState);
+}
+
+uint16 FMQTTConnection::PopId()
+{
+	return MessageIdGenerator.Pop();
+}
+
+void FMQTTConnection::PushId(const uint16 InValue)
+{
+	MessageIdGenerator.Push(InValue);
+}
+
+bool FMQTTConnection::IsConnected() const
+{
+	return GetState() == EMQTTState::Connected && (Socket && Socket->GetConnectionState() == SCS_Connected);
+}
+
+#undef LOCTEXT_NAMESPACE // "MQTTConnection"
\ No newline at end of file
diff --git a/Engine/Plugins/Protocols/MQTT/Source/MQTTCore/Private/MQTTConnection.h b/Engine/Plugins/Protocols/MQTT/Source/MQTTCore/Private/MQTTConnection.h
new file mode 100644
--- /dev/null
+++ b/Engine/Plugins/Protocols/MQTT/Source/MQTTCore/Private/MQTTConnection.h
@@ -0,0 +1,165 @@
+﻿// Copyright Epic Games, Inc. All Rights Reserved.
+
+#pragma once
+
+#include "CoreMinimal.h"
+#include "MQTTCoreLog.h"
+#include "MQTTOperations.h"
+#include "MQTTShared.h"
+#include "MQTTSharedInternal.h"
+#include "Sockets.h"
+#include "Containers/Queue.h"
+#include "Containers/StaticArray.h"
+#include "Containers/StaticBitArray.h"
+#include "HAL/Runnable.h"
+#include "Misc/SingleThreadRunnable.h"
+#include "Serialization/ArrayReader.h"
+
+/** A wrapped unique (per client) id generator. */
+class FMQTTPacketIdGenerator
+{
+public:
+	FMQTTPacketIdGenerator();
+	
+	uint16 Pop();
+
+	void Push(const uint16 InValue);
+
+protected:
+	static constexpr uint16 Max = TNumericLimits<uint16>::Max();
+	std::atomic<uint16> Min; // Current index of next slot to be returned
+	std::atomic<uint16> Position; // Current index of slot ready to go
+	TStaticArray<uint16, Max> ValueStore;
+};
+
+/** Socket connection. */
+class MQTTCORE_API FMQTTConnection final
+	: public FRunnable
+	, public FSingleThreadRunnable
+	, public TSharedFromThis<FMQTTConnection, ESPMode::ThreadSafe>
+{
+protected:
+	// Allows MakeShared with private constructor
+	friend class SharedPointerInternals::TIntrusiveReferenceController<FMQTTConnection, ESPMode::ThreadSafe>;
+	
+	/** Constructor. Hidden on purpose, use Create instead. */
+	explicit FMQTTConnection(
+		FSocket& InSocket,
+		TSharedRef<FInternetAddr> InAddr);
+
+public:
+	virtual ~FMQTTConnection() override;
+
+	static TSharedPtr<FMQTTConnection, ESPMode::ThreadSafe> TryCreate(const TSharedPtr<FInternetAddr>& InAddr);
+
+	template <class OperationType>
+	auto QueueOperation(typename OperationType::FRequestType&& InRequestPacket) -> TFuture<typename OperationType::FResponseType>;
+
+	EMQTTState GetState() const { return MQTTState.load(); }
+	void SetState(const EMQTTState InState);
+
+	/** Get new message Id. */
+	uint16 PopId();
+
+	/** Add message Id back to pool. */
+	void PushId(uint16 InValue);
+
+	bool IsConnected() const;
+
+	DECLARE_DELEGATE_OneParam(FOnMessage, const FMQTTPublishPacket& /* Packet */)
+	FOnMessage& OnMessage() { return OnMessageDelegate; }	
+	
+protected:
+	//~ Begin FRunnable implementation
+	virtual bool Init() override;
+	virtual uint32 Run() override;
+	virtual void Stop() override;
+	virtual void Exit() override;
+	//~ End FRunnable implementation
+
+	//~ Begin FSingleThreadRunnable implementation
+	virtual void Tick() override;
+	virtual class FSingleThreadRunnable* GetSingleThreadInterface() override;
+	//~ End FSingleThreadRunnable implementation
+
+	/** Updates the socket connection and resulting state. Return true if connected. */
+	bool UpdateConnection();
+	void ProcessMessages();
+
+	void OpenSocket();
+	
+	FCriticalSection OutgoingQueueLock;
+	FCriticalSection IncomingQueueLock;
+	TQueue<TUniquePtr<IMQTTOperation>, EQueueMode::Spsc> OutgoingOperations;
+	TUniquePtr<TMQTTOperation<FMQTTConnectPacket, FMQTTConnectAckPacket>> PendingConnectOperation;
+	TUniquePtr<IMQTTOperation> PendingPingOperation; 
+	TMap<uint16, TUniquePtr<IMQTTOperation>> PendingIncomingOperations; // Waiting for ack's, etc.
+	
+	std::atomic<EMQTTSocketState> SocketState;
+	std::atomic<EMQTTState> MQTTState;
+
+	FCriticalSection MessageIdGeneratorLock;
+	FMQTTPacketIdGenerator MessageIdGenerator;
+
+private:
+	void AbandonOperations();
+	void HandleMessage(EMQTTPacketType InMessageType, const TSharedRef<FArrayReader>& InPacketReader);
+	void HandleConnectAcknowledge(const TSharedRef<FArrayReader>& InPacketReader);
+	void HandleSubscribeAcknowledge(const TSharedRef<FArrayReader>& InPacketReader);
+	void HandleUnsubscribeAcknowledge(const TSharedRef<FArrayReader>& InPacketReader);
+	void HandlePingResponse(const TSharedRef<FArrayReader>& InPacketReader);
+	void HandlePublish(const TSharedRef<FArrayReader>& InPacketReader);
+	void HandlePublishAcknowledge(const TSharedRef<FArrayReader>& InPacketReader);
+	void HandlePublishReceived(const TSharedRef<FArrayReader>& InPacketReader);
+	void HandlePublishRelease(const TSharedRef<FArrayReader>& InPacketReader);	
+	void HandlePublishComplete(const TSharedRef<FArrayReader>& InPacketReader);
+
+private:
+	/** The network socket (incoming). */
+	FSocket* Socket;
+
+	/** The endpoint internet addr */
+	TSharedPtr<FInternetAddr> EndpointInternetAddr;
+
+	/** The thread object. */
+	FRunnableThread* Thread;
+
+	/** The receiver thread's name. */
+	FString ThreadName;
+
+	/** Used when any message is received (publish packet). */
+	FOnMessage OnMessageDelegate;
+
+	float SocketConnectionTimeout = 10.0f; // in seconds
+	float MQTTConnectionTimeout = 10.0f; // in seconds
+	double LastConnectRequestTime = 0.0f; // in seconds
+	float KeepAliveInterval = 60.0f; // in seconds
+	double LastKeepAliveTime = 0.0f;
+};
+
+template <class OperationType>
+auto FMQTTConnection::QueueOperation(typename OperationType::FRequestType&& InRequestPacket)
+	-> TFuture<typename OperationType::FResponseType>
+{
+	using RequestType = typename OperationType::FRequestType;
+	using ResponseType = typename OperationType::FResponseType;
+
+#if UE_BUILD_DEBUG
+	if(GetState() == EMQTTState::Stopping)
+	{
+		UE_LOG(LogMQTTCore, Warning, TEXT("Tried to queue operation while stopping: %s"), ANSI_TO_TCHAR(typeid(OperationType).name()));
+	}
+
+	UE_LOG(LogMQTTCore, Verbose, TEXT("Queued Operation: %s"), ANSI_TO_TCHAR(typeid(OperationType).name()));
+#endif
+
+	// const auto Operation = MakeShared<TMQTTOperation<RequestType, ResponseType>, ESPMode::ThreadSafe>(MoveTemp(InRequestPacket), MessageIdGenerator->Poke());
+	TUniquePtr<OperationType> Operation = MakeUnique<OperationType>(MoveTemp(InRequestPacket), InRequestPacket.PacketId);
+	auto Future = Operation->GetFuture();
+	{
+		FScopeLock Lock(&OutgoingQueueLock);
+		OutgoingOperations.Enqueue(MoveTemp(Operation));
+	}
+	
+	return Future;
+}
diff --git a/Engine/Plugins/Protocols/MQTT/Source/MQTTCore/Private/MQTTCoreLog.cpp b/Engine/Plugins/Protocols/MQTT/Source/MQTTCore/Private/MQTTCoreLog.cpp
new file mode 100644
--- /dev/null
+++ b/Engine/Plugins/Protocols/MQTT/Source/MQTTCore/Private/MQTTCoreLog.cpp
@@ -0,0 +1,5 @@
+// Copyright Epic Games, Inc. All Rights Reserved.
+
+#include "MQTTCoreLog.h"
+
+DEFINE_LOG_CATEGORY(LogMQTTCore);
diff --git a/Engine/Plugins/Protocols/MQTT/Source/MQTTCore/Private/MQTTCoreLog.h b/Engine/Plugins/Protocols/MQTT/Source/MQTTCore/Private/MQTTCoreLog.h
new file mode 100644
--- /dev/null
+++ b/Engine/Plugins/Protocols/MQTT/Source/MQTTCore/Private/MQTTCoreLog.h
@@ -0,0 +1,7 @@
+// Copyright Epic Games, Inc. All Rights Reserved.
+
+#pragma once
+
+#include "CoreMinimal.h"
+
+DECLARE_LOG_CATEGORY_EXTERN(LogMQTTCore, VeryVerbose, All);
diff --git a/Engine/Plugins/Protocols/MQTT/Source/MQTTCore/Private/MQTTCoreModule.cpp b/Engine/Plugins/Protocols/MQTT/Source/MQTTCore/Private/MQTTCoreModule.cpp
new file mode 100644
--- /dev/null
+++ b/Engine/Plugins/Protocols/MQTT/Source/MQTTCore/Private/MQTTCoreModule.cpp
@@ -0,0 +1,268 @@
+﻿// Copyright Epic Games, Inc. All Rights Reserved.
+
+#include "MQTTCoreModule.h"
+
+#include "GameDelegates.h"
+#include "IMQTTClient.h"
+#include "MQTTClient.h"
+#include "MQTTClientSettings.h"
+#include "MQTTCoreLog.h"
+#include "MQTTShared.h"
+#include "Algo/Count.h"
+
+#include "HAL/IConsoleManager.h"
+#include "Interfaces/IPluginManager.h"
+#include "Misc/OutputDeviceMemory.h"
+#include "Misc/Paths.h"
+#include "Serialization/BufferWriter.h"
+#include "Templates/SharedPointer.h"
+
+void FMQTTCoreModule::StartupModule()
+{
+	RegisterConsoleCommands();
+	FGameDelegates::Get().GetEndPlayMapDelegate().AddRaw(this, &FMQTTCoreModule::OnEndPlay);
+}
+
+void FMQTTCoreModule::ShutdownModule()
+{
+	FGameDelegates::Get().GetEndPlayMapDelegate().RemoveAll(this);
+	UnregisterConsoleCommands();
+	MQTTClients.Empty();
+}
+
+TSharedPtr<IMQTTClient, ESPMode::ThreadSafe> FMQTTCoreModule::GetOrCreateClient(const bool bForceNew)
+{
+	const UMQTTClientSettings* Settings = GetDefault<UMQTTClientSettings>();
+	return GetOrCreateClient(Settings->DefaultURL, bForceNew);
+}
+
+TSharedPtr<IMQTTClient, ESPMode::ThreadSafe> FMQTTCoreModule::GetOrCreateClient(const FMQTTURL& InURL, const bool bForceNew)
+{
+	const FGuid ClientId = InURL.ToGuid();
+
+	// Finds existing or creates new Client based on unique URL
+	if(const TWeakPtr<IMQTTClient, ESPMode::ThreadSafe>* ExistingClient = MQTTClients.Find(ClientId))
+	{
+		if(!bForceNew)
+		{
+			return ExistingClient->Pin();
+		}
+
+		ExistingClient->Pin().Reset();
+		MQTTClients.Remove(ClientId);
+	}
+
+	TSharedPtr<IMQTTClient, ESPMode::ThreadSafe> CreatedClient = MQTTClients.Add(ClientId, MakeShared<FMQTTClient, ESPMode::ThreadSafe>(InURL)).Pin();
+	ensure(MQTTClients.Num() > 0);
+	UE_LOG(LogMQTTCore, Display, TEXT("Created new Client, Num: %d"), MQTTClients.Num());
+
+	return CreatedClient;
+}
+
+void FMQTTCoreModule::RegisterConsoleCommands()
+{ 
+	ConsoleCommands.Add(MakeUnique<FAutoConsoleCommand>(
+		TEXT("MQTT.GetOrCreateClient"),
+		TEXT("Create the MQTT client"),
+		FConsoleCommandWithArgsDelegate::CreateRaw(this, &FMQTTCoreModule::OnCreateClientCommand)
+		));
+
+	ConsoleCommands.Add(MakeUnique<FAutoConsoleCommand>(
+		TEXT("MQTT.DestroyClient"),
+		TEXT("Destroy the MQTT client"),
+		FConsoleCommandWithArgsDelegate::CreateRaw(this, &FMQTTCoreModule::OnDestroyClientCommand)
+		));
+
+	ConsoleCommands.Add(MakeUnique<FAutoConsoleCommand>(
+		TEXT("MQTT.ClientSubscribe"),
+		TEXT("Subscribe on the topic"),
+		FConsoleCommandWithArgsDelegate::CreateRaw(this, &FMQTTCoreModule::OnClientSubscribeCommand)
+		));
+
+	ConsoleCommands.Add(MakeUnique<FAutoConsoleCommand>(
+		TEXT("MQTT.ClientUnsubscribe"),
+		TEXT("Unsubscribe from the topic"),
+		FConsoleCommandWithArgsDelegate::CreateRaw(this, &FMQTTCoreModule::OnClientUnsubscribeCommand)
+		));
+
+	ConsoleCommands.Add(MakeUnique<FAutoConsoleCommand>(
+		TEXT("MQTT.ClientPublish"),
+		TEXT("Publish to the topic"),
+		FConsoleCommandWithArgsDelegate::CreateRaw(this, &FMQTTCoreModule::OnClientPublishCommand)
+		));
+}
+
+void FMQTTCoreModule::UnregisterConsoleCommands()
+{
+	for (TUniquePtr<FAutoConsoleCommand>& Command : ConsoleCommands)
+	{
+		Command.Reset();
+	}
+}
+
+void FMQTTCoreModule::OnCreateClientCommand(const TArray<FString>& InArgs)
+{
+	if (InArgs.Num() != 2)
+	{
+		return;
+	}
+
+	const FString& Host = InArgs[0];
+
+	int32 Port = -1;
+	LexFromString(Port, *InArgs[1]);
+
+	if (Port == -1)
+	{
+		return;
+	}
+
+	GetOrCreateClient(FMQTTURL(Host, Port));
+}
+
+void FMQTTCoreModule::OnDestroyClientCommand(const TArray<FString>& InArgs)
+{
+	if (InArgs.Num() != 1)
+	{
+		return;
+	}
+
+	const FGuid ClientId = FGuid(InArgs[0]);
+
+	UE_LOG(LogMQTTCore, Display, TEXT("DestroyClient: %s"), *ClientId.ToString());
+
+	// If the input is valid, this client *should* exist in MQTTClients
+	ensureAlwaysMsgf(MQTTClients.Num() > 0, TEXT("Attempted to DestroyClient but MQTTClients was empty."));
+
+	if(const TWeakPtr<IMQTTClient, ESPMode::ThreadSafe>* FoundClient = MQTTClients.Find(ClientId))
+	{
+		// if only 1 shared ref, it's the cached client unused by any active objects, so destroy
+		if(FoundClient->Pin().GetSharedReferenceCount() <= 1)
+		{
+			MQTTClients.Remove(ClientId);	
+		}
+	}
+}
+
+void FMQTTCoreModule::OnClientSubscribeCommand(const TArray<FString>& InArgs)
+{
+	if (InArgs.Num() != 3)
+	{
+		return;
+	}
+	
+	const FGuid ClientId = FGuid(InArgs[0]);
+	if (!ClientId.IsValid())
+	{
+		return;
+	}
+
+	const FString SubscriptionPattern = InArgs[1];
+
+	int32 QualityOfService = -1;
+	LexFromString(QualityOfService, *InArgs[2]);
+	if (QualityOfService == -1)
+	{
+		return;
+	}
+	
+	const TSharedPtr<IMQTTClient, ESPMode::ThreadSafe> MQTTClient = FindClient(ClientId);
+	if (!MQTTClient.IsValid())
+	{
+		return;
+	}
+
+	MQTTClient->Subscribe({TPairInitializer<FString, EMQTTQualityOfService>{SubscriptionPattern, static_cast<EMQTTQualityOfService>(QualityOfService)}});
+}
+
+void FMQTTCoreModule::OnClientUnsubscribeCommand(const TArray<FString>& InArgs)
+{
+	if (InArgs.Num() != 2)
+	{
+		return;
+	}
+	
+	const FGuid ClientId = FGuid(InArgs[0]);
+	if (!ClientId.IsValid())
+	{
+		return;
+	}
+
+	const FString SubscriptionPattern = InArgs[1];
+
+	const TSharedPtr<IMQTTClient, ESPMode::ThreadSafe> MQTTClient = FindClient(ClientId);
+	if (!MQTTClient.IsValid())
+	{
+		return;
+	}
+
+	MQTTClient->Unsubscribe({SubscriptionPattern});
+}
+
+void FMQTTCoreModule::OnClientPublishCommand(const TArray<FString>& InArgs)
+{
+	if (InArgs.Num() != 5)
+	{
+		return;
+	}
+	
+	const FGuid ClientId = FGuid(InArgs[0]);
+	if (!ClientId.IsValid())
+	{
+		return;
+	}
+
+	const FString Topic = InArgs[1];
+	const FString Payload = InArgs[2];
+
+	int32 QualityOfService = -1;
+	LexFromString(QualityOfService, *InArgs[3]);
+	if (QualityOfService == -1)
+	{
+		return;
+	}
+
+	bool bRetain = false;	
+	LexFromString(bRetain, *InArgs[4]);
+
+	const TSharedPtr<IMQTTClient, ESPMode::ThreadSafe> MQTTClient = FindClient(ClientId);
+	if (!MQTTClient.IsValid())
+	{
+		return;
+	}
+
+	MQTTClient->Publish(Topic, Payload, static_cast<EMQTTQualityOfService>(QualityOfService), bRetain);
+}
+
+void FMQTTCoreModule::OnEndPlay()
+{
+	for(auto& MQTTClient : MQTTClients)
+	{
+		MQTTClient.Value.Reset();
+	}
+	MQTTClients.Empty();
+}
+
+int32 FMQTTCoreModule::GetClientNum() const
+{
+	// Only count valid clients (ignore stale pointers)
+	return Algo::CountIf(MQTTClients,
+	[](const TPair<FGuid, TWeakPtr<IMQTTClient, ESPMode::ThreadSafe>>& InIdClientPair)
+	{
+		return InIdClientPair.Value.IsValid();
+	});
+}
+
+TSharedPtr<IMQTTClient, ESPMode::ThreadSafe> FMQTTCoreModule::FindClient(const FGuid InId)
+{
+	check(InId.IsValid());
+
+	if(TSharedPtr<IMQTTClient, ESPMode::ThreadSafe> FoundClient = MQTTClients.Find(InId)->Pin())
+	{
+		return FoundClient;		
+	}
+
+	return nullptr; 
+}
+
+IMPLEMENT_MODULE(FMQTTCoreModule, MQTTCore);
diff --git a/Engine/Plugins/Protocols/MQTT/Source/MQTTCore/Private/MQTTCoreModule.h b/Engine/Plugins/Protocols/MQTT/Source/MQTTCore/Private/MQTTCoreModule.h
new file mode 100644
--- /dev/null
+++ b/Engine/Plugins/Protocols/MQTT/Source/MQTTCore/Private/MQTTCoreModule.h
@@ -0,0 +1,47 @@
+﻿// Copyright Epic Games, Inc. All Rights Reserved.
+
+#pragma once
+
+#include "IMQTTCoreModule.h"
+
+class FAutoConsoleCommand;
+
+class FMQTTCoreModule : public IMQTTCoreModule
+{
+public:
+	//~ Begin IModuleInterface
+	virtual void StartupModule() override;
+	virtual void ShutdownModule() override;
+	//~ End IModuleInterface
+
+	//~ Begin IMQTTCoreModule
+	virtual TSharedPtr<IMQTTClient, ESPMode::ThreadSafe> GetOrCreateClient(bool bForceNew = false) override;
+	virtual TSharedPtr<IMQTTClient, ESPMode::ThreadSafe> GetOrCreateClient(const FMQTTURL& InURL, bool bForceNew = false) override;
+	//~ End IMQTTCoreModule
+
+private:
+	/** Register console commands. */
+	void RegisterConsoleCommands();
+
+	/** Unregister console commands. */
+	void UnregisterConsoleCommands();
+
+	void OnCreateClientCommand(const TArray<FString>& InArgs);	
+	void OnDestroyClientCommand(const TArray<FString>& InArgs);
+
+	void OnClientSubscribeCommand(const TArray<FString>& InArgs);
+	void OnClientUnsubscribeCommand(const TArray<FString>& InArgs);
+	void OnClientPublishCommand(const TArray<FString>& InArgs);
+
+	void OnEndPlay();
+
+protected:
+	TSharedPtr<IMQTTClient, ESPMode::ThreadSafe> FindClient(FGuid InId);
+	virtual int32 GetClientNum() const override;
+
+private:
+	TMap<FGuid, TWeakPtr<IMQTTClient, ESPMode::ThreadSafe>> MQTTClients;
+
+	/** Console commands handles. */
+	TArray<TUniquePtr<FAutoConsoleCommand>> ConsoleCommands;
+};
diff --git a/Engine/Plugins/Protocols/MQTT/Source/MQTTCore/Private/MQTTCoreStats.cpp b/Engine/Plugins/Protocols/MQTT/Source/MQTTCore/Private/MQTTCoreStats.cpp
new file mode 100644
--- /dev/null
+++ b/Engine/Plugins/Protocols/MQTT/Source/MQTTCore/Private/MQTTCoreStats.cpp
@@ -0,0 +1,3 @@
+// Copyright Epic Games, Inc. All Rights Reserved.
+
+#include "MQTTCoreStats.h"
diff --git a/Engine/Plugins/Protocols/MQTT/Source/MQTTCore/Private/MQTTCoreStats.h b/Engine/Plugins/Protocols/MQTT/Source/MQTTCore/Private/MQTTCoreStats.h
new file mode 100644
--- /dev/null
+++ b/Engine/Plugins/Protocols/MQTT/Source/MQTTCore/Private/MQTTCoreStats.h
@@ -0,0 +1,10 @@
+// Copyright Epic Games, Inc. All Rights Reserved.
+
+#pragma once
+
+#include "CoreMinimal.h"
+
+#include "Stats/Stats.h"
+#include "Stats/Stats2.h"
+
+DECLARE_STATS_GROUP(TEXT("MQTT"), STATGROUP_MQTT, STATCAT_Advanced);
diff --git a/Engine/Plugins/Protocols/MQTT/Source/MQTTCore/Private/MQTTOperations.cpp b/Engine/Plugins/Protocols/MQTT/Source/MQTTCore/Private/MQTTOperations.cpp
new file mode 100644
--- /dev/null
+++ b/Engine/Plugins/Protocols/MQTT/Source/MQTTCore/Private/MQTTOperations.cpp
@@ -0,0 +1,614 @@
+﻿// Copyright Epic Games, Inc. All Rights Reserved.
+
+#include "MQTTOperations.h"
+
+#include "Algo/IndexOf.h"
+
+REGISTER_MQTT_ARCHIVE(FMQTTFixedHeader)
+REGISTER_MQTT_ARCHIVE(FMQTTConnectPacket)
+REGISTER_MQTT_ARCHIVE(FMQTTConnectAckPacket)
+REGISTER_MQTT_ARCHIVE(FMQTTPublishPacket)
+REGISTER_MQTT_ARCHIVE(FMQTTPublishAckPacket)
+REGISTER_MQTT_ARCHIVE(FMQTTPublishReceivedPacket)
+REGISTER_MQTT_ARCHIVE(FMQTTPublishReleasePacket)
+REGISTER_MQTT_ARCHIVE(FMQTTPublishCompletePacket)
+REGISTER_MQTT_ARCHIVE(FMQTTSubscribePacket)
+REGISTER_MQTT_ARCHIVE(FMQTTSubscribeAckPacket)
+REGISTER_MQTT_ARCHIVE(FMQTTUnsubscribePacket)
+REGISTER_MQTT_ARCHIVE(FMQTTUnsubscribeAckPacket)
+REGISTER_MQTT_ARCHIVE(FMQTTPingRequestPacket)
+REGISTER_MQTT_ARCHIVE(FMQTTPingResponsePacket)
+REGISTER_MQTT_ARCHIVE(FMQTTDisconnectPacket)
+
+namespace Internal
+{
+	static TMap<EMQTTPacketType, const TCHAR*> PacketTypeNames =
+	{
+		{ EMQTTPacketType::Connect, TEXT("Connect") },
+		{ EMQTTPacketType::ConnectAcknowledge, TEXT("ConnectAcknowledge") },
+			
+		{ EMQTTPacketType::Publish, TEXT("Publish") },
+		{ EMQTTPacketType::PublishAcknowledge, TEXT("PublishAcknowledge") },
+		{ EMQTTPacketType::PublishReceived, TEXT("PublishReceive") },
+		{ EMQTTPacketType::PublishRelease, TEXT("PublishRelease") },
+		{ EMQTTPacketType::PublishComplete, TEXT("PublishComplete") },
+
+		{ EMQTTPacketType::Subscribe, TEXT("Subscribe") },
+		{ EMQTTPacketType::SubscribeAcknowledge, TEXT("SubscribeAcknowledge") },
+
+		{ EMQTTPacketType::Unsubscribe, TEXT("Unsubscribe") },
+		{ EMQTTPacketType::UnsubscribeAcknowledge, TEXT("UnsubscribeAcknowledge") },
+
+		{ EMQTTPacketType::PingRequest, TEXT("PingRequest") },
+		{ EMQTTPacketType::PingResponse, TEXT("PingResponse") },
+
+		{ EMQTTPacketType::Disconnect, TEXT("Disconnect") },
+		{ EMQTTPacketType::Authorize, TEXT("Authorize") },
+
+		/** non-spec entries */
+		{ EMQTTPacketType::None, TEXT("None") },
+	};
+}
+
+const TCHAR* MQTT::GetMQTTPacketTypeName(EMQTTPacketType InMessageType)
+{
+	return Internal::PacketTypeNames[InMessageType];
+}
+
+template <typename SizeType = uint16>
+void SerializeString(FArchive& Ar, const FString& InString)
+{
+	SizeType StringLength = InString.Len();
+	const char* String = TCHAR_TO_UTF8(*InString);
+	Ar << StringLength;
+	Ar.Serialize((void*)String, StringLength);
+}
+
+template <typename SizeType = uint16>
+void DeserializeString(FArchive& Ar, FString& OutString)
+{
+	SizeType StringLength;
+	Ar << StringLength;
+	TArray<char> Str;
+	Str.Reserve(StringLength);
+	Ar.Serialize(static_cast<void*>(Str.GetData()), StringLength);
+	OutString = UTF8_TO_TCHAR(Str.GetData());
+	OutString.LeftInline(StringLength); // @todo: fix, not sure why i get garbage end character	
+}
+
+/** Reads from Archive and writes to Header. */
+uint8 FMQTTFixedHeader::ReadHeader(FArchive& InArchive, FMQTTFixedHeader& OutHeader)
+{
+	ensure(InArchive.IsLoading());
+	
+	InArchive.SetByteSwapping(true);
+
+	InArchive << OutHeader.PacketType;
+
+	uint32 MessageLength = 0;
+	uint8 EncodedByte;
+	unsigned int RemainingMultiplier = 1;
+	for(int32 Idx = 0; Idx < 4; ++Idx)
+	{
+		if(!InArchive.AtEnd())
+		{
+			InArchive << EncodedByte;
+			MessageLength += (EncodedByte & 127) * RemainingMultiplier;
+			RemainingMultiplier *= 128;
+			if((EncodedByte & 128) == 0)
+			{
+				OutHeader.RemainingLength = MessageLength;
+				return OutHeader.LengthFieldSize = Idx + 1;
+			}
+		}
+	}
+
+	return false;
+}
+
+/** Reads from Header and writes to Archive. */
+uint8 FMQTTFixedHeader::WriteHeader(FArchive& InArchive, const FMQTTFixedHeader& InHeader)
+{
+	ensure(InArchive.IsSaving());
+	
+	InArchive.SetByteSwapping(true);
+	
+	uint8 RemainingBytes[5];
+	uint32 RemainingLength = InHeader.RemainingLength;
+	uint32 RemainingCount = 0;
+	do
+	{
+		uint8 EncodedByte = RemainingLength % 128;
+		RemainingLength /= 128;
+		if(RemainingLength > 0)
+		{
+			EncodedByte |= 128;
+		}
+		RemainingBytes[RemainingCount] = EncodedByte;
+		RemainingCount++;
+	}
+	while (RemainingLength > 0 && RemainingCount < 5);
+	
+	ensure(RemainingCount != 5);
+
+	uint8 PacketType = InHeader.PacketType;
+	InArchive << PacketType;
+	for(uint32 Idx = 0; Idx < RemainingCount; ++Idx)
+	{
+		InArchive << RemainingBytes[Idx];
+	}
+
+	return RemainingCount;
+}
+
+int64 FMQTTFixedHeader::HeaderSize() const
+{
+	return 1 + LengthFieldSize;
+}
+
+int64 FMQTTFixedHeader::TotalSize() const
+{
+	// PacketType + Length (varint) + Remaining
+	return 1 + LengthFieldSize + RemainingLength;
+}
+
+void FMQTTFixedHeader::Serialize(FArchive& Ar)
+{
+	if(Ar.IsLoading())
+	{
+		ReadHeader(Ar, *this);
+	}
+	else if(Ar.IsSaving())
+	{
+		WriteHeader(Ar, *this);
+	}
+}
+
+void FMQTTConnectPacket::Serialize(FArchive& Ar)
+{
+	Ar.SetByteSwapping(true);
+
+	if(Ar.IsLoading())
+	{
+		// ReadHeader(Ar, *this);
+	}
+	else if(Ar.IsSaving())
+	{
+		const bool bHasClientId = !ClientId.IsEmpty();
+		const bool bHasUserName = !UserName.IsEmpty();
+		const bool bHasPassword = !Password.IsEmpty();
+		bool bHasWill = !WillTopic.IsEmpty();
+
+		check(bHasPassword ? bHasUserName : true);
+
+		constexpr uint32 HeaderLength = 10;
+		uint32 PayloadLength = 0;
+
+		PayloadLength += bHasClientId ? 2 + ClientId.Len() : 2;
+		PayloadLength += bHasWill ? 2 + WillTopic.Len() + 2 + WillMessage.Len() : 0;
+		PayloadLength += bHasUserName ? 2 + UserName.Len() : 0;
+		PayloadLength += bHasPassword ? 2 + Password.Len() : 0;
+
+		FixedHeader.SetPacketType(EMQTTPacketType::Connect);
+		FixedHeader.RemainingLength = HeaderLength + PayloadLength;
+		Ar << FixedHeader;
+
+		SerializeString(Ar, "MQTT");
+
+		uint8 ProtocolLevel = static_cast<uint8>(EMQTTProtocolLevel::MQTT311); // MQTT 3.11
+		Ar << ProtocolLevel;
+
+		uint8 Flags = (bCleanSession & 0x1) << 1;
+		if(bHasWill)
+		{
+			Flags = Flags | ((static_cast<int32>(WillQoS) & 0x3) << 3 | (bHasWill & 0x1) << 2);
+			if(bRetainWill)
+			{
+				Flags |= ((bRetainWill & 0x1) << 5);
+			}
+		}
+
+		if(bHasUserName)
+		{
+			Flags |= 0x1 << 7;
+		}
+
+		if(bHasPassword)
+		{
+			Flags |= 0x1 << 6;
+		}
+	
+		Ar << Flags;
+		Ar << KeepAliveSeconds;
+
+		// Payload
+		if(bHasClientId)
+		{
+			SerializeString(Ar, ClientId);
+		}
+		else
+		{
+			static uint16 Zero = 0;
+			Ar << Zero;
+		}
+
+		if(bHasWill)
+		{
+			SerializeString(Ar, WillTopic);
+			SerializeString(Ar, WillMessage);
+		}
+
+		if(bHasUserName)
+		{
+			SerializeString(Ar, UserName);
+			Flags |= 0x1 << 7;
+		}
+
+		if(bHasPassword)
+		{
+			SerializeString(Ar, Password);
+			Flags |= 0x1 << 6;
+		}
+	}
+}
+
+void FMQTTConnectAckPacket::Serialize(FArchive& Ar)
+{
+	Ar.SetByteSwapping(true);
+	
+	if(Ar.IsLoading())
+	{
+		Ar << FixedHeader;
+		Ar << Flags;
+		Ar << ReturnCode;
+	}
+	else if(Ar.IsSaving())
+	{
+	}
+}
+
+auto FMQTTPublishPacket::Serialize(FArchive& Ar) -> void
+{
+	Ar.SetByteSwapping(true);
+	// @todo: disable bRetain if retain not available
+
+	if(Ar.IsLoading())
+	{
+		// Mark read position here
+		int64 HeaderSize = Ar.Tell();
+		Ar << FixedHeader;
+
+		bIsDuplicate = (FixedHeader.PacketType & 0x08) >> 3;
+		QoS = static_cast<EMQTTQualityOfService>((FixedHeader.PacketType & 0x06) >> 1);
+		bRetain = (FixedHeader.PacketType & 0x01);
+
+		DeserializeString(Ar, Topic);
+
+		if(QoS > EMQTTQualityOfService::Once)
+		{
+			Ar << PacketId;
+		}
+		// Calculate size
+		HeaderSize = Ar.Tell() - HeaderSize;
+
+		int32 PayloadLength = FixedHeader.RemainingLength - HeaderSize;
+		if(PayloadLength > 0)
+		{
+			PayloadLength += 2;
+			Payload.Init(0, PayloadLength);
+			for(int Idx = 0; Idx < PayloadLength; ++Idx)
+			{
+				Ar << Payload[Idx];
+			}
+		}
+		auto X = Ar.Tell();
+		auto Y = X;
+	}
+	else if(Ar.IsSaving())
+	{
+		uint32 PacketLength = 2 + Payload.Num();
+		if(!Topic.IsEmpty())
+		{
+			PacketLength += Topic.Len();	
+		}
+
+		FixedHeader.PacketType = static_cast<uint8>(static_cast<uint8>(EMQTTPacketType::Publish) | static_cast<uint8>(((bIsDuplicate ? 1 : 0) & 0x1) << 3) | (static_cast<uint8>(QoS) << 1) | static_cast<uint8>(bRetain ? 1 : 0));
+		FixedHeader.RemainingLength = PacketLength + (QoS > EMQTTQualityOfService::Once ? 2 : 0);
+		Ar << FixedHeader;
+
+		if(!Topic.IsEmpty())
+		{
+			SerializeString(Ar, Topic);	
+		}
+		else
+		{
+			static uint16 Zero;
+			Ar << Zero;	 
+		}
+
+		// [MQTT-2.3.1-5] of v3.1.1 spec
+		if(QoS > EMQTTQualityOfService::Once)
+		{
+			Ar << PacketId;
+		}
+
+		if(Payload.Num() > 0)
+		{
+			for(int Idx = 0; Idx < Payload.Num(); ++Idx)
+			{
+				Ar << Payload[Idx];
+			}
+		}
+	}
+}
+
+void FMQTTPublishAckPacket::Serialize(FArchive& Ar)
+{
+	Ar.SetByteSwapping(true);
+
+	if(Ar.IsLoading())
+	{
+		Ar << FixedHeader;
+		Ar << PacketId;
+	}
+	else if(Ar.IsSaving())
+	{
+	}
+}
+
+void FMQTTPublishReceivedPacket::Serialize(FArchive& Ar)
+{
+	Ar.SetByteSwapping(true);
+
+	if(Ar.IsLoading())
+	{
+		Ar << FixedHeader;
+		Ar << PacketId;
+	}
+	else if(Ar.IsSaving())
+	{
+		FixedHeader.PacketType = static_cast<uint8>(EMQTTPacketType::PublishReceived);
+		FixedHeader.Reserved = 2;
+		FixedHeader.RemainingLength = 2;
+		Ar << FixedHeader;
+		Ar << PacketId;
+	}
+}
+
+void FMQTTPublishReleasePacket::Serialize(FArchive& Ar)
+{
+	Ar.SetByteSwapping(true);
+	
+	if(Ar.IsLoading())
+	{
+		Ar << FixedHeader;
+		Ar << PacketId;
+	}
+	else if(Ar.IsSaving())
+	{
+		FixedHeader.PacketType = static_cast<uint8>(EMQTTPacketType::PublishRelease) | 2;
+		FixedHeader.Reserved = 2;
+		FixedHeader.RemainingLength = 2;
+		Ar << FixedHeader;
+		Ar << PacketId;
+	}
+}
+
+void FMQTTPublishCompletePacket::Serialize(FArchive& Ar)
+{
+	if(Ar.IsLoading())
+	{
+		Ar << FixedHeader;
+		Ar << PacketId;
+	}
+	else if(Ar.IsSaving())
+	{
+		FixedHeader.PacketType = static_cast<uint8>(EMQTTPacketType::PublishComplete);
+		FixedHeader.Reserved = 2;
+		FixedHeader.RemainingLength = 2;
+		Ar << FixedHeader;
+		Ar << PacketId;
+	}
+}
+
+void FMQTTSubscribePacket::Serialize(FArchive& Ar)
+{
+	Ar.SetByteSwapping(true);
+
+	if(Ar.IsLoading())
+	{
+	}
+	else if(Ar.IsSaving())
+	{
+		// @todo: disable bRetain if retain not available	
+		uint32 PacketLength = 2;
+		for(FString& Topic : Topics)
+		{
+			ensure(Topic.Len() <= TNumericLimits<uint16>::Max());
+			PacketLength += 2 + Topic.Len() + 1;
+		}
+
+		FixedHeader.PacketType = static_cast<uint8>(static_cast<uint8>(EMQTTPacketType::Subscribe) | 1 << 1);
+		FixedHeader.RemainingLength = PacketLength;
+		Ar << FixedHeader;
+		Ar << PacketId;
+
+		for(auto Idx = 0; Idx < Topics.Num(); ++Idx)
+		{
+			SerializeString(Ar, Topics[Idx]);
+			Ar << RequestedQoS[Idx];
+		}
+	}
+}
+
+void FMQTTSubscribeAckPacket::Serialize(FArchive& Ar)
+{
+	Ar.SetByteSwapping(true);
+
+	if(Ar.IsLoading())
+	{
+		Ar << FixedHeader;
+		Ar << PacketId;
+
+		for(auto Idx = 0; Idx < FixedHeader.RemainingLength - 2; ++Idx)
+		{
+			Ar << ReturnCodes.Emplace_GetRef();
+		}
+	}
+	else if(Ar.IsSaving())
+	{
+
+	}
+}
+
+void FMQTTUnsubscribePacket::Serialize(FArchive& Ar)
+{
+	Ar.SetByteSwapping(true);
+
+	if(Ar.IsLoading())
+	{
+	}
+	else if(Ar.IsSaving())
+	{
+		// @todo: disable bRetain if retain not available	
+		uint32 PacketLength = 2;
+		for(FString& Topic : Topics)
+		{
+			ensure(Topic.Len() < TNumericLimits<uint16>::Max());
+			PacketLength += 2 + Topic.Len();
+		}
+
+		FixedHeader.PacketType = static_cast<uint8>(static_cast<uint8>(EMQTTPacketType::Unsubscribe) | 1 << 1);
+		FixedHeader.RemainingLength = PacketLength;
+		Ar << FixedHeader;
+		Ar << PacketId;
+
+		for(FString& Topic : Topics)
+		{
+			SerializeString(Ar, Topic);
+		}
+	}
+}
+
+void FMQTTUnsubscribeAckPacket::Serialize(FArchive& Ar)
+{
+	Ar.SetByteSwapping(true);
+
+	if(Ar.IsLoading())
+	{
+		Ar << FixedHeader;		
+		Ar << PacketId;
+	}
+	else if(Ar.IsSaving())
+	{
+		FixedHeader.PacketType = static_cast<uint8>(EMQTTPacketType::UnsubscribeAcknowledge);
+		FixedHeader.Reserved = 2;
+		FixedHeader.RemainingLength = 2;
+		Ar << FixedHeader;
+		Ar << PacketId;
+	}
+}
+
+void FMQTTPingRequestPacket::Serialize(FArchive& Ar)
+{
+	Ar.SetByteSwapping(true);
+
+	if(Ar.IsLoading())
+	{
+		Ar << FixedHeader;
+	}
+	else if(Ar.IsSaving())
+	{
+		FixedHeader.PacketType = static_cast<uint8>(EMQTTPacketType::PingRequest);
+		FixedHeader.Reserved = 2;
+		FixedHeader.RemainingLength = 0;
+		Ar << FixedHeader;
+	}
+}
+
+void FMQTTPingResponsePacket::Serialize(FArchive& Ar)
+{
+	Ar.SetByteSwapping(true);
+
+	if(Ar.IsLoading())
+	{
+		Ar << FixedHeader;
+	}
+	else if(Ar.IsSaving())
+	{
+		FixedHeader.SetPacketType(EMQTTPacketType::PingResponse);
+		FixedHeader.RemainingLength = 2;
+		Ar << FixedHeader;
+	}
+}
+
+void FMQTTDisconnectPacket::Serialize(FArchive& Ar)
+{
+	Ar.SetByteSwapping(true);
+
+	if(Ar.IsLoading())
+	{
+		Ar << FixedHeader;
+	}
+	else if(Ar.IsSaving())
+	{
+		FixedHeader.PacketType = static_cast<uint8>(EMQTTPacketType::Disconnect);
+		Ar << FixedHeader;
+	}
+}
+
+void FMQTTConnectOperation::Abandon()
+{
+	if(bCompleted.load() == false)
+	{
+		Promise.SetValue(FMQTTConnectAckPacket(EMQTTConnectReturnCode::SocketError));
+	}
+}
+
+void FMQTTPublishOperationQoS2::Pack(FBufferArchive& InArchive)
+{
+	switch (CurrentPacketType)
+	{
+	case EMQTTPacketType::Publish:
+		TMQTTOperation<FMQTTPublishPacket, FMQTTPublishCompletePacket>::Pack(InArchive);
+		break;
+		
+	case EMQTTPacketType::PublishReceived:
+		InArchive << PublishReceivedPacket;
+		break;
+		
+	case EMQTTPacketType::PublishRelease:
+		InArchive << PublishReleasePacket;
+		break;
+
+	default:
+		checkNoEntry();
+	}
+}
+
+uint16 FMQTTPublishOperationQoS2::CompleteStep(FMQTTPublishReceivedPacket&& InResponse)
+{
+	checkf(PendingPacketType == EMQTTPacketType::PublishReceived, TEXT("Unexpected packet type!"));
+	PublishReceivedPacket = MoveTemp(InResponse);
+	PublishReleasePacket.PacketId = PublishReceivedPacket.PacketId;
+	NextStep();
+	NextStep();
+	return PacketId;
+}
+
+uint16 FMQTTPublishOperationQoS2::CompleteStep(FMQTTPublishReleasePacket&& InResponse)
+{
+	checkf(PendingPacketType == EMQTTPacketType::PublishRelease, TEXT("Unexpected packet type!"));
+	PublishReleasePacket = MoveTemp(InResponse);
+	NextStep();
+	NextStep();
+	return PacketId;
+}
+
+// get next request/response pair
+void FMQTTPublishOperationQoS2::NextStep()
+{
+	const EMQTTPacketType NextPacketType = StepPacketTypes[Algo::IndexOf(StepPacketTypes, PendingPacketType) + 1];
+	Swap(CurrentPacketType, PendingPacketType);
+	PendingPacketType = NextPacketType;
+}
diff --git a/Engine/Plugins/Protocols/MQTT/Source/MQTTCore/Private/MQTTOperations.h b/Engine/Plugins/Protocols/MQTT/Source/MQTTCore/Private/MQTTOperations.h
new file mode 100644
--- /dev/null
+++ b/Engine/Plugins/Protocols/MQTT/Source/MQTTCore/Private/MQTTOperations.h
@@ -0,0 +1,650 @@
+﻿// Copyright Epic Games, Inc. All Rights Reserved.
+
+#pragma once
+
+#include "CoreMinimal.h"
+#include "MQTTProtocol.h"
+#include "MQTTSharedInternal.h"
+
+#pragma region Packets
+
+struct FMQTTFixedHeader
+{
+	uint8 PacketType = static_cast<uint8>(EMQTTPacketType::Connect);
+	uint8 Reserved = 0;
+	uint8 LengthFieldSize = 1; // 1 - 4 depending on int precision
+	uint16 RemainingLength = 0; // Length of variable header (10 bytes) + payload length
+	uint16 TotalLength = 0;
+
+	EMQTTPacketType GetPacketType() const
+	{
+		return static_cast<EMQTTPacketType>(PacketType);
+	}
+
+	void SetPacketType(EMQTTPacketType InPacketType)
+	{
+		PacketType = static_cast<uint8>(InPacketType);
+	}
+
+	int64 HeaderSize() const;
+	int64 TotalSize() const;
+
+	void Serialize(FArchive& Ar);
+	friend FArchive& operator<<(FArchive& Ar, FMQTTFixedHeader& InPacket);
+
+protected:
+	/** Reads from Archive and writes to Header. Returns size of Length field (1-4 bytes). */
+	static uint8 ReadHeader(FArchive& InArchive, FMQTTFixedHeader& OutHeader);
+
+	/** Reads from Header and writes to Archive. Returns size of Length field (1-4 bytes). */
+	static uint8 WriteHeader(FArchive& InArchive, const FMQTTFixedHeader& InHeader);
+};
+
+struct FMQTTPacket
+{
+	// Not needed for certain packets, but practical for templating purposes
+	uint16 PacketId = 0;
+	
+	FMQTTPacket() = default;	
+	virtual ~FMQTTPacket() = default;
+
+	/** Get's the "RemainingLength" field. Mainly used for testing/verification. */
+	int32 GetLength() const { return FixedHeader.RemainingLength; }
+
+protected:
+	FMQTTFixedHeader FixedHeader;
+};
+
+struct FMQTTConnectPacket : FMQTTPacket
+{
+	constexpr static EMQTTPacketType PacketType = EMQTTPacketType::Connect;
+
+	explicit FMQTTConnectPacket(
+		FString InClientId,
+		const uint16 InKeepAliveSeconds = 60,
+		FString InUserName = "",
+		FString InPassword = "",
+		const bool bInCleanSession = false,
+		const bool bInRetainWill = false,
+		FString InWillTopic = "",
+		FString InWillMessage = "",
+		const EMQTTQualityOfService InWillQoS = EMQTTQualityOfService::Once)
+		: ClientId(MoveTemp(InClientId))
+		, KeepAliveSeconds(InKeepAliveSeconds)
+		, UserName(MoveTemp(InUserName))
+		, Password(MoveTemp(InPassword))
+		, bCleanSession(bInCleanSession)
+		, bRetainWill(bInRetainWill)
+		, WillTopic(MoveTemp(InWillTopic))
+		, WillMessage(MoveTemp(InWillMessage))
+		, WillQoS(InWillQoS)
+	{
+	}
+
+	FString ClientId;
+	uint16 KeepAliveSeconds = 60;
+	FString UserName;
+	FString Password;
+	bool bCleanSession;
+	bool bRetainWill;
+	FString WillTopic;
+	FString WillMessage;
+	EMQTTQualityOfService WillQoS;
+
+	void Serialize(FArchive& Ar);
+
+	friend FArchive& operator<<(FArchive& Ar, FMQTTConnectPacket& InPacket);
+};
+
+struct FMQTTConnectAckPacket : FMQTTPacket
+{
+	constexpr static EMQTTPacketType PacketType = EMQTTPacketType::ConnectAcknowledge;
+
+	FMQTTConnectAckPacket() = default;
+
+	/** Short circuit when there's a connection error. */
+	explicit FMQTTConnectAckPacket(const EMQTTConnectReturnCode& InReturnCode)
+		: Flags(0)
+		, ReturnCode(InReturnCode)
+	{
+	}
+
+	uint8 Flags; // if clean session from client is 1, this should be 0, if from client is 0, this should be 1 if client session stored, or 0 if not
+	EMQTTConnectReturnCode ReturnCode;
+
+	void Serialize(FArchive& Ar);
+
+	friend FArchive& operator<<(FArchive& Ar, FMQTTConnectAckPacket& InPacket);
+};
+
+// @todo: for MQTT 5, properties, end of the message, pg25 of spec
+struct FMQTTPublishPacket : FMQTTPacket
+{
+	constexpr static EMQTTPacketType PacketType = EMQTTPacketType::Publish;
+
+	FMQTTPublishPacket() = default;
+
+	explicit FMQTTPublishPacket(
+		const uint16 InPacketId,
+		FString InTopic,
+		const TArray<uint8> InPayload,
+		const EMQTTQualityOfService InQoS,
+		const bool bInRetain)
+		: QoS(InQoS)
+		, bRetain(bInRetain)
+		, Topic(std::move(InTopic))
+		, Payload(InPayload)
+	{
+		PacketId = InPacketId;
+	}
+
+	bool bIsDuplicate = false; // msg is a retry
+	EMQTTQualityOfService QoS = EMQTTQualityOfService::Once;
+	bool bRetain = false;
+	FString Topic = ""; // utf8, no wildcards
+
+	TArray<uint8> Payload; // RemainingLength - VariableLength, can be 0
+	uint8 Response = 0; // depends on QoS, 0 = none, 1 = PUBACK, 2 = PUBREC
+
+	void Serialize(FArchive& Ar);
+
+	friend FArchive& operator<<(FArchive& Ar, FMQTTPublishPacket& InPacket);
+};
+
+struct FMQTTPublishAckPacket : FMQTTPacket
+{
+	constexpr static EMQTTPacketType PacketType = EMQTTPacketType::PublishAcknowledge;
+	uint8 Length = 2;
+
+	void Serialize(FArchive& Ar);
+
+	friend FArchive& operator<<(FArchive& Ar, FMQTTPublishAckPacket& InPacket);
+};
+
+struct FMQTTPublishReceivedPacket : FMQTTPacket
+{
+	constexpr static EMQTTPacketType PacketType = EMQTTPacketType::PublishReceived;
+
+	void Serialize(FArchive& Ar);
+
+	friend FArchive& operator<<(FArchive& Ar, FMQTTPublishReceivedPacket& InPacket);
+};
+
+struct FMQTTPublishReleasePacket : FMQTTPacket
+{
+	constexpr static EMQTTPacketType PacketType = EMQTTPacketType::PublishRelease;
+
+	FMQTTPublishReleasePacket() = default;
+
+	explicit FMQTTPublishReleasePacket(const uint16 InPacketId)
+	{
+		PacketId = InPacketId;
+	}
+
+	void Serialize(FArchive& Ar);
+
+	friend FArchive& operator<<(FArchive& Ar, FMQTTPublishReleasePacket& InPacket);
+};
+
+struct FMQTTPublishCompletePacket : FMQTTPacket
+{
+	constexpr static EMQTTPacketType PacketType = EMQTTPacketType::PublishComplete;
+	FString Topic;
+	TArray<uint8> Payload;
+
+	void Serialize(FArchive& Ar);
+
+	friend FArchive& operator<<(FArchive& Ar, FMQTTPublishCompletePacket& InPacket);
+};
+
+struct FMQTTSubscribePacket : FMQTTPacket
+{
+	constexpr static EMQTTPacketType PacketType = EMQTTPacketType::Subscribe;
+
+	FMQTTSubscribePacket() = default;
+
+	explicit FMQTTSubscribePacket(
+		const uint16 InPacketId,
+		const TArray<FString>& InTopics,
+		const TArray<EMQTTQualityOfService>& InRequestedQoS)
+		: Topics(InTopics)
+		, RequestedQoS(InRequestedQoS)
+	{
+		PacketId = InPacketId;
+	}
+
+	TArray<FString> Topics; // utf8
+	TArray<EMQTTQualityOfService> RequestedQoS;
+
+	void Serialize(FArchive& Ar);
+
+	friend FArchive& operator<<(FArchive& Ar, FMQTTSubscribePacket& InPacket);
+};
+
+struct FMQTTSubscribeAckPacket : FMQTTPacket
+{
+	constexpr static EMQTTPacketType PacketType = EMQTTPacketType::SubscribeAcknowledge;
+	TArray<uint8> ReturnCodes;
+
+	void Serialize(FArchive& Ar);
+
+	friend FArchive& operator<<(FArchive& Ar, FMQTTSubscribeAckPacket& InPacket);
+};
+
+struct FMQTTUnsubscribePacket : FMQTTPacket
+{
+	constexpr static EMQTTPacketType PacketType = EMQTTPacketType::Unsubscribe;
+	TArray<FString> Topics; // utf8
+
+	FMQTTUnsubscribePacket() = default;
+
+	explicit FMQTTUnsubscribePacket(
+		const uint16 InPacketId,
+		const TArray<FString>& InTopics)
+		: Topics(InTopics)
+	{
+		PacketId = InPacketId;
+	}
+
+	void Serialize(FArchive& Ar);
+
+	friend FArchive& operator<<(FArchive& Ar, FMQTTUnsubscribePacket& InPacket);
+};
+
+struct FMQTTUnsubscribeAckPacket : FMQTTPacket
+{
+	constexpr static EMQTTPacketType PacketType = EMQTTPacketType::UnsubscribeAcknowledge;
+
+	void Serialize(FArchive& Ar);
+
+	friend FArchive& operator<<(FArchive& Ar, FMQTTUnsubscribeAckPacket& InPacket);
+};
+
+struct FMQTTPingRequestPacket : FMQTTPacket
+{
+	constexpr static EMQTTPacketType PacketType = EMQTTPacketType::PingRequest;
+
+	void Serialize(FArchive& Ar);
+
+	friend FArchive& operator<<(FArchive& Ar, FMQTTPingRequestPacket& InPacket);
+};
+
+struct FMQTTPingResponsePacket : FMQTTPacket
+{
+	constexpr static EMQTTPacketType PacketType = EMQTTPacketType::PingResponse;
+
+	void Serialize(FArchive& Ar);
+
+	friend FArchive& operator<<(FArchive& Ar, FMQTTPingResponsePacket& InPacket);
+};
+
+struct FMQTTDisconnectPacket : FMQTTPacket
+{
+	constexpr static EMQTTPacketType PacketType = EMQTTPacketType::Disconnect;
+
+	void Serialize(FArchive& Ar);
+
+	friend FArchive& operator<<(FArchive& Ar, FMQTTDisconnectPacket& InPacket);
+};
+
+#pragma endregion Packets
+
+#pragma region Operations
+
+/** An operation defines a message flow and encapsulates one or more packets. */
+class IMQTTOperation
+{
+	constexpr static EMQTTPacketType NullPacketType = EMQTTPacketType::None;
+
+public:
+	explicit IMQTTOperation(
+		const uint16 InPacketId = 0,
+		const EMQTTState& InStateDependency = EMQTTState::Connected)
+		: StateDependency(InStateDependency)
+#if WITH_EDITORONLY_DATA
+		, TimeQueued(FDateTime::UtcNow())
+		, TimeExecuted(0)
+#endif
+		, PacketId(InPacketId)
+		, bCompleted(false)
+	{
+	}	
+
+	virtual ~IMQTTOperation() = default;
+
+	virtual EMQTTPacketType GetPacketType() const
+	{
+		return NullPacketType;
+	}
+
+	virtual EMQTTPacketType GetPendingPacketType() const
+	{
+		return NullPacketType;
+	}
+
+	const uint16& GetPacketId() const
+	{
+		return PacketId;
+	}
+
+	/** Serialize to provided archive. */
+	virtual void Pack(FBufferArchive& InArchive)
+	{
+	}
+
+	/** Flags task as complete and returns the PacketId, if any. */
+	virtual uint16 Complete()
+	{
+		return PacketId;
+	}
+
+	/** Abandon task. */
+	virtual void Abandon()
+	{
+	}
+
+	EMQTTState StateDependency = EMQTTState::None;
+
+#if WITH_EDITORONLY_DATA
+	FDateTime TimeQueued;
+	FDateTime TimeExecuted;
+
+	uint32 GetTimeToExecution() const
+	{
+		return (TimeExecuted - TimeQueued).GetFractionMilli();
+	}
+
+	void MarkExecuted()
+	{
+		TimeExecuted = FDateTime::UtcNow();
+	}
+#endif
+
+protected:
+	uint16 PacketId = 0;
+	std::atomic<bool> bCompleted;
+};
+
+/** Persisted as Session State. */
+class IMQTTPersistentOperation
+{
+};
+
+template <typename RequestType, typename ResponseType>
+class TMQTTOperation : public IMQTTOperation
+{
+public:
+	typedef RequestType FRequestType;
+	typedef ResponseType FResponseType;
+
+	explicit TMQTTOperation(
+		RequestType&& InRequest,
+		const uint16 InPacketId = 0,
+		const EMQTTState& InStateDependency = EMQTTState::Connected)
+		: IMQTTOperation(InPacketId, InStateDependency)
+		, Packet(MoveTemp(InRequest))
+	{
+	}
+
+	virtual ~TMQTTOperation() override;
+
+	virtual EMQTTPacketType GetPacketType() const override
+	{
+		return RequestType::PacketType;
+	}
+
+	virtual EMQTTPacketType GetPendingPacketType() const override
+	{
+		return ResponseType::PacketType;
+	}
+
+	/** Serialize to provided archive. */
+	virtual void Pack(FBufferArchive& InArchive) override
+	{
+		InArchive << Packet;
+	}
+
+	TFuture<ResponseType> GetFuture()
+	{
+		return Promise.GetFuture();
+	}
+
+	virtual uint16 Complete() override;
+	virtual uint16 Complete(ResponseType&& InResponse);
+
+	/** Abandon task. */
+	virtual void Abandon() override;
+
+protected:
+	RequestType Packet;
+	TPromise<ResponseType> Promise;
+};
+
+template <typename RequestType>
+class TMQTTOperation<RequestType, void> : public IMQTTOperation
+{
+public:
+	typedef RequestType FRequestType;
+	typedef void FResponseType;
+
+	explicit TMQTTOperation(
+		RequestType&& InRequest,
+		const uint16 InPacketId = 0,
+		const EMQTTState& InStateDependency = EMQTTState::Connected)
+		: IMQTTOperation(InPacketId, InStateDependency)
+		, Packet(MoveTemp(InRequest))
+	{
+	}
+
+	virtual EMQTTPacketType GetPacketType() const override
+	{
+		return RequestType::PacketType;
+	}
+
+	/** Serialize to provided archive. */
+	virtual void Pack(FBufferArchive& InArchive) override
+	{
+		InArchive << Packet;
+	}
+
+	TFuture<void> GetFuture()
+	{
+		return Promise.GetFuture();
+	}
+
+	virtual uint16 Complete() override;
+
+	/** Abandon task. */
+	virtual void Abandon() override
+	{
+		if(bCompleted.load() == false)
+		{
+			Promise.SetValue();
+		}
+	}
+
+protected:
+	RequestType Packet;
+	TPromise<void> Promise;
+};
+
+template <typename RequestType, typename ResponseType>
+TMQTTOperation<RequestType, ResponseType>::~TMQTTOperation()
+{
+}
+
+template <typename RequestType, typename ResponseType>
+uint16 TMQTTOperation<RequestType, ResponseType>::Complete()
+{
+	bCompleted.store(true);
+	Promise.SetValue({});
+	return PacketId;
+}
+
+template <typename RequestType>
+uint16 TMQTTOperation<RequestType, void>::Complete()
+{
+	bCompleted.store(true);
+	Promise.SetValue();
+	return PacketId;
+}
+
+template <typename RequestType, typename ResponseType>
+uint16 TMQTTOperation<RequestType, ResponseType>::Complete(ResponseType&& InResponse)
+{
+	bCompleted.store(true);
+	Promise.SetValue(MoveTemp(InResponse));
+	return PacketId;
+}
+
+template <typename RequestType, typename ResponseType>
+void TMQTTOperation<RequestType, ResponseType>::Abandon()
+{
+	if(bCompleted.load() == false)
+	{
+		Promise.SetValue({});
+	}
+}
+
+class FMQTTConnectOperation : public TMQTTOperation<FMQTTConnectPacket, FMQTTConnectAckPacket>
+{
+public:
+	explicit FMQTTConnectOperation(FMQTTConnectPacket&& InRequest, const uint16 InPacketId = 0)
+		: TMQTTOperation<FMQTTConnectPacket, FMQTTConnectAckPacket>(MoveTemp(InRequest), InPacketId, EMQTTState::None)
+	{
+	}
+
+	virtual void Abandon() override;
+};
+
+class FMQTTDisconnectOperation : public TMQTTOperation<FMQTTDisconnectPacket, void>
+{
+public:
+	explicit FMQTTDisconnectOperation(FMQTTDisconnectPacket&& InRequest, const uint16 InPacketId = 0)
+		: TMQTTOperation<FMQTTDisconnectPacket, void>(MoveTemp(InRequest), InPacketId, EMQTTState::None)
+	{
+	}
+};
+
+class FMQTTPublishOperationQoS0 : public TMQTTOperation<FMQTTPublishPacket, void>
+{
+public:
+	explicit FMQTTPublishOperationQoS0(FMQTTPublishPacket&& InRequest, const uint16 InPacketId = 0)
+		: TMQTTOperation<FMQTTPublishPacket, void>(MoveTemp(InRequest), InPacketId, EMQTTState::Connected)
+	{
+	}
+};
+
+class FMQTTPublishOperationQoS1 : public TMQTTOperation<FMQTTPublishPacket, FMQTTPublishAckPacket>
+{
+public:
+	explicit FMQTTPublishOperationQoS1(FMQTTPublishPacket&& InRequest, const uint16 InPacketId = 0)
+		: TMQTTOperation<FMQTTPublishPacket, FMQTTPublishAckPacket>(MoveTemp(InRequest), InPacketId, EMQTTState::Connected)
+	{
+	}
+};
+
+// Wraps all steps FMQTTPublishPacket -> FMQTTPublishCompletePacket
+class FMQTTPublishOperationQoS2 final
+	: public TMQTTOperation<FMQTTPublishPacket, FMQTTPublishCompletePacket>,
+	public IMQTTPersistentOperation
+{
+public:
+	explicit FMQTTPublishOperationQoS2(FMQTTPublishPacket&& InRequest, const uint16 InPacketId = 0)
+		: TMQTTOperation<FMQTTPublishPacket, FMQTTPublishCompletePacket>(MoveTemp(InRequest), InPacketId, EMQTTState::Connected)
+		, PublishReceivedPacket()
+		, PublishReleasePacket()
+	{
+	}
+
+	/** Serialize to provided archive. Packet serialized depends on current step. */
+	virtual void Pack(FBufferArchive& InArchive) override;
+
+	virtual EMQTTPacketType GetPacketType() const override
+	{
+		return CurrentPacketType;
+	}
+
+	virtual EMQTTPacketType GetPendingPacketType() const override
+	{
+		return PendingPacketType;
+	}
+
+	uint16 CompleteStep(FMQTTPublishReceivedPacket&& InResponse);
+	uint16 CompleteStep(FMQTTPublishReleasePacket&& InResponse);
+
+protected:
+	static constexpr EMQTTPacketType StepPacketTypes[] = {EMQTTPacketType::PublishReceived, EMQTTPacketType::PublishRelease, EMQTTPacketType::PublishComplete};
+	EMQTTPacketType CurrentPacketType = EMQTTPacketType::Publish;
+	EMQTTPacketType PendingPacketType = EMQTTPacketType::PublishReceived;
+
+	FMQTTPublishReceivedPacket PublishReceivedPacket;
+	FMQTTPublishReleasePacket PublishReleasePacket;
+
+	void NextStep();
+};
+
+// FMQTTPublishPacket -> FMQTTPublishReceivedPacket
+class FMQTTPublishOperationQoS2_Step1
+	: public TMQTTOperation<FMQTTPublishPacket, FMQTTPublishReceivedPacket>,
+	public IMQTTPersistentOperation
+{
+public:
+	explicit FMQTTPublishOperationQoS2_Step1(FMQTTPublishPacket&& InRequest, const uint16 InPacketId = 0)
+		: TMQTTOperation<FMQTTPublishPacket, FMQTTPublishReceivedPacket>(MoveTemp(InRequest), InPacketId, EMQTTState::Connected)
+	{
+	}
+};
+
+// FMQTTPublishReceivedPacket -> FMQTTPublishReleasePacket
+class FMQTTPublishOperationQoS2_Step2
+	: public TMQTTOperation<FMQTTPublishReceivedPacket, FMQTTPublishReleasePacket>,
+	public IMQTTPersistentOperation
+{
+public:
+	explicit FMQTTPublishOperationQoS2_Step2(FMQTTPublishReceivedPacket&& InRequest, const uint16 InPacketId = 0)
+		: TMQTTOperation<FMQTTPublishReceivedPacket, FMQTTPublishReleasePacket>(MoveTemp(InRequest), InPacketId, EMQTTState::Connected)
+	{
+	}
+};
+
+// FMQTTPublishReleasePacket -> FMQTTPublishCompletePacket
+class FMQTTPublishOperationQoS2_Step3
+	: public TMQTTOperation<FMQTTPublishReleasePacket, FMQTTPublishCompletePacket>,
+	public IMQTTPersistentOperation
+{
+public:
+	explicit FMQTTPublishOperationQoS2_Step3(FMQTTPublishReleasePacket&& InRequest, const uint16 InPacketId = 0)
+		: TMQTTOperation<FMQTTPublishReleasePacket, FMQTTPublishCompletePacket>(MoveTemp(InRequest), InPacketId, EMQTTState::Connected)
+	{
+	}
+};
+
+class FMQTTSubscribeOperation : public TMQTTOperation<FMQTTSubscribePacket, FMQTTSubscribeAckPacket>
+{
+public:
+	explicit FMQTTSubscribeOperation(FMQTTSubscribePacket&& InRequest, const uint16 InPacketId = 0)
+		: TMQTTOperation<FMQTTSubscribePacket, FMQTTSubscribeAckPacket>(MoveTemp(InRequest), InPacketId, EMQTTState::Connected)
+	{
+	}
+};
+
+class FMQTTUnsubscribeOperation : public TMQTTOperation<FMQTTUnsubscribePacket, FMQTTUnsubscribeAckPacket>
+{
+public:
+	explicit FMQTTUnsubscribeOperation(FMQTTUnsubscribePacket&& InRequest, const uint16 InPacketId = 0)
+		: TMQTTOperation<FMQTTUnsubscribePacket, FMQTTUnsubscribeAckPacket>(MoveTemp(InRequest), InPacketId, EMQTTState::Connected)
+	{
+	}
+};
+
+class FMQTTPingOperation : public TMQTTOperation<FMQTTPingRequestPacket, FMQTTPingResponsePacket>
+{
+public:
+	explicit FMQTTPingOperation(FMQTTPingRequestPacket&& InRequest, const uint16 InPacketId = 0)
+		: TMQTTOperation<FMQTTPingRequestPacket, FMQTTPingResponsePacket>(MoveTemp(InRequest), InPacketId, EMQTTState::Connected)
+	{
+	}
+};
+
+#pragma endregion Operations
diff --git a/Engine/Plugins/Protocols/MQTT/Source/MQTTCore/Private/MQTTProtocol.cpp b/Engine/Plugins/Protocols/MQTT/Source/MQTTCore/Private/MQTTProtocol.cpp
new file mode 100644
--- /dev/null
+++ b/Engine/Plugins/Protocols/MQTT/Source/MQTTCore/Private/MQTTProtocol.cpp
@@ -0,0 +1,3 @@
+﻿// Copyright Epic Games, Inc. All Rights Reserved.
+
+#include "MQTTProtocol.h"
diff --git a/Engine/Plugins/Protocols/MQTT/Source/MQTTCore/Private/MQTTShared.cpp b/Engine/Plugins/Protocols/MQTT/Source/MQTTCore/Private/MQTTShared.cpp
new file mode 100644
--- /dev/null
+++ b/Engine/Plugins/Protocols/MQTT/Source/MQTTCore/Private/MQTTShared.cpp
@@ -0,0 +1,380 @@
+﻿// Copyright Epic Games, Inc. All Rights Reserved.
+
+#include "MQTTShared.h"
+
+#include "SocketSubsystem.h"
+#include "Algo/AllOf.h"
+
+#define LOCTEXT_NAMESPACE "MQTT"
+
+bool FMQTTURL::GetAddress(ISocketSubsystem* InSocketSubsystem, TSharedPtr<FInternetAddr>& OutAddress) const
+{
+	OutAddress = InSocketSubsystem->CreateInternetAddr();
+
+	// Probably Ip
+	if(Host.Contains("."))
+	{
+		bool bIsValid = false;
+		OutAddress->SetIp(*Host, bIsValid);
+		if(!bIsValid)
+		{
+			return false;
+		}
+		OutAddress->SetPort(Port);
+		return true;
+	}
+	// Probably hostname
+	{
+		FAddressInfoResult AddressInfo = InSocketSubsystem->GetAddressInfo(
+			*Host,
+			nullptr,
+			EAddressInfoFlags::AllResultsWithMapping | EAddressInfoFlags::OnlyUsableAddresses,
+			NAME_None);
+		
+		if (AddressInfo.ReturnCode == SE_NO_ERROR && AddressInfo.Results.Num() > 0)
+		{
+			// found, valid
+			OutAddress = AddressInfo.Results[0].Address->Clone();
+			OutAddress->SetPort(Port);
+			return OutAddress->IsValid();
+		}
+		else
+		{
+			// todo: log error
+		}
+	}
+	
+	return false;
+}
+
+// @note: based on WaterUtils::StringToGuid()
+FGuid FMQTTURL::ToGuid() const
+{
+	// Compute a 128-bit hash based on the string and use that as a GUID :
+	const FTCHARToUTF8 Converted(*ToString());
+	FMD5 MD5Gen;
+	MD5Gen.Update((const uint8*)Converted.Get(), Converted.Length());
+	uint32 Digest[4];
+	MD5Gen.Final((uint8*)Digest);
+
+	// FGuid::NewGuid() creates a version 4 UUID (at least on Windows), which will have the top 4 bits of the
+	// second field set to 0100. We'll set the top bit to 1 in the GUID we create, to ensure that we can never
+	// have a collision with other implicitly-generated GUIDs.
+	Digest[1] |= 0x80000000;
+	return FGuid(Digest[0], Digest[1], Digest[2], Digest[3]);
+}
+
+// Host is required, all others are optional
+bool FMQTTURL::IsValid() const
+{
+	// User, Pass max length = 65535 bytes
+	return !bFlaggedInvalid
+		&& !Host.IsEmpty()	// Host must be set
+		&& (!Username.IsEmpty() ? Username.Len() > TNumericLimits<uint16>::Max() - 2 : true)
+		&& (!Password.IsEmpty() ? Password.Len() > TNumericLimits<uint16>::Max() - 2 : true)
+		&& (!Password.IsEmpty() ? !Username.IsEmpty() : true); // Password needs Username too
+}
+
+bool FMQTTURL::IsValid(FText& OutMessage) const
+{
+	if(bFlaggedInvalid)
+	{
+		OutMessage = LOCTEXT("FlaggedInvalid", "URL has been flagged as invalid.");
+		return false;
+	}
+		
+	if(Host.IsEmpty())
+	{
+		OutMessage = LOCTEXT("InvalidHost", "Host (name or ip) is invalid.");
+		return false;
+	}
+
+	if(!Password.IsEmpty() && Username.IsEmpty())
+	{
+		OutMessage = LOCTEXT("PasswordWithoutUsername", "Password set without a Username.");
+		return false;
+	}
+
+	return true;
+}
+
+FMQTTURL FMQTTURL::FromString(const FString& InURL)
+{
+	FMQTTURL Result;
+	FString ParsedStr = InURL;
+	
+	// Regex support isn't great, but here's the pattern: (?<Scheme>mqtt|mqtts|tcp|ssl)(?::\\/\\/)?(?<Username>\\w+)?(?<Password>(?::)\\w+)?@?(?<Host>\\w+)+:?(?<Port>\\d+)?
+
+	// Scheme
+	int32 DelimeterIdx = ParsedStr.Find(TEXT("://"));
+	if(DelimeterIdx >= 0)
+	{
+		const FString SchemeStr = ParsedStr.Left(DelimeterIdx);
+
+		static UEnum* SchemeEnum = FindObject<UEnum>(nullptr, TEXT("MQTTCore.EMQTTScheme"));
+		int64 SchemeValueIdx = SchemeEnum->GetValueByNameString(SchemeStr);
+		if(SchemeValueIdx != INDEX_NONE)
+		{
+			Result.Scheme = static_cast<EMQTTScheme>(SchemeValueIdx);	
+		}
+
+		ParsedStr.RightInline(ParsedStr.Len() - DelimeterIdx - 3);
+	}
+
+	// User+Pass
+	DelimeterIdx = ParsedStr.Find(TEXT("@"));
+	if(DelimeterIdx >= 0)
+	{
+		FString UserStr = ParsedStr.Left(DelimeterIdx);
+
+		// Pass
+		DelimeterIdx = UserStr.Find(TEXT(":"));
+		if(DelimeterIdx >= 0)
+		{
+			const FString PassStr = UserStr.RightChop(DelimeterIdx + 1);
+			Result.Password = PassStr;
+			
+			UserStr.LeftInline(DelimeterIdx);
+		}
+
+		Result.Username = UserStr;
+		ParsedStr.RightInline(ParsedStr.Len() - DelimeterIdx - UserStr.Len() - 2);
+	}
+
+	// Host+Port
+	DelimeterIdx = ParsedStr.Find(TEXT(":"));
+	if(DelimeterIdx >= 0)
+	{
+		const FString PortStr = ParsedStr.Right(ParsedStr.Len() - DelimeterIdx - 1);
+		Result.Port = FCString::Atoi64(*PortStr);
+
+		ParsedStr.LeftInline(DelimeterIdx);
+	}
+
+	Result.Host = ParsedStr;
+
+	return Result;
+}
+
+// mqtt[s]://[username][:password]@host.domain[:port]
+FString FMQTTURL::ToString() const
+{ 
+	FString Result;
+
+	// scheme ie. mqtt://
+	Result.Append(UEnum::GetValueAsString(TEXT("MQTTCore.EMQTTScheme"), Scheme).ToLower() + "://");
+
+	// user+pass
+	if(!Username.IsEmpty())
+	{
+		Result.Append(Username);
+
+		if(!Password.IsEmpty())
+		{
+			Result.Append(":" + Password);
+		}
+
+		Result.Append("@");
+	}
+
+	// hostname
+	Result.Append(Host);
+
+	// port
+	if(Port > 0)
+	{
+		Result.Appendf(TEXT(":%i"), Port);
+	}
+
+	return Result;
+}
+
+void FMQTTTopic::MakeValid()
+{
+	Path.ReplaceInline(TEXT("\\"), TEXT("/")); // replace backslashes with forward
+}
+
+bool FMQTTTopic::IsValid(const FString& InTopic)
+{
+	if(InTopic.Len() == 0
+		|| InTopic.Len() > 65535)
+	{
+		return false;
+	}
+	if(InTopic.Len() == 0
+		|| InTopic.Contains(TEXT("\\")))
+	{
+		return false;
+	}
+
+	for(auto CharIdx = 0; CharIdx < InTopic.Len(); ++CharIdx)
+	{
+		if(InTopic[CharIdx] == FMQTTTopic::SingleLevelWildcardChar
+			// if not first char last should be /
+			&& ((CharIdx > 0 && CharIdx < InTopic.Len() - 1 && InTopic[CharIdx - 1] != '/')
+			// or if not last char next should be /
+			|| (CharIdx != InTopic.Len() - 1 && InTopic[CharIdx + 1] != '/')))
+		{
+			return false;
+		}
+
+		if(InTopic[CharIdx] == FMQTTTopic::MultiLevelWildcardChar
+			// if not first char last should be /
+			&& ((CharIdx > 0 && CharIdx < InTopic.Len() - 1 && InTopic[CharIdx + 1] != '/')
+			// or if not last char
+			|| (CharIdx != InTopic.Len() - 1)))
+		{
+			return false;
+		}
+	}
+
+	return true;
+}
+
+bool FMQTTTopic::IsValid() const
+{
+	return FMQTTTopic::IsValid(Path);
+}
+
+FString FMQTTTopic::ToString() const
+{
+	return Path;
+}
+
+bool FMQTTTopicFilter::IsValid() const
+{
+	return FMQTTTopic::IsValid();
+}
+
+void FMQTTTopicFilter::MakeValid()
+{
+	FMQTTTopic::MakeValid();
+}
+
+FString FMQTTTopicFilter::ToString() const
+{
+	return FMQTTTopic::ToString();
+}
+
+bool FMQTTTopicTemplate::PopulateTemplate(const FStringFormatNamedArguments& InArgs, FMQTTTopic& OutTopic) const
+{
+	FString PopulatedPath = FString::Format(*Path, InArgs);
+	OutTopic.Path = PopulatedPath;
+	return true;
+}
+
+bool FMQTTSubscription::Matches(const FString& InTopic, const FString& InSubscription)
+{
+	size_t SubscriptionPos = 0;
+	const size_t SubscriptionLength = InSubscription.Len();
+	size_t TopicPos = 0;
+	const size_t TopicLength = InTopic.Len();
+
+	while (SubscriptionPos < SubscriptionLength && TopicPos < TopicLength)
+	{
+		// Char match
+		if (InSubscription[SubscriptionPos] == InTopic[TopicPos])
+		{
+			// End of topic
+			if (TopicPos == TopicLength - 1)
+			{
+				// Subscription leaf matches due to wildcard
+				// Remaining chars are '/#'
+				if ((SubscriptionPos == SubscriptionLength - 3
+					&& InSubscription[SubscriptionPos + 1] == FMQTTTopic::TopicLevelSeparatorChar
+					&& InSubscription[SubscriptionPos + 2] == FMQTTTopic::MultiLevelWildcardChar)
+
+				// This char is '/', next is '#'
+					|| (SubscriptionPos == SubscriptionLength - 2
+					&& InSubscription[SubscriptionPos] == FMQTTTopic::TopicLevelSeparatorChar
+					&& InSubscription[SubscriptionPos + 1] == FMQTTTopic::MultiLevelWildcardChar))
+				{
+					return true;
+				}
+			}
+
+			SubscriptionPos++;
+			TopicPos++;
+
+			// At end of both, complete match
+			if (SubscriptionPos == SubscriptionLength
+				&& TopicPos == TopicLength)
+			{
+				return true;
+			} 
+
+			// End of topic, subscription ends with wildcard 
+			if (TopicPos == TopicLength
+				&& SubscriptionPos == SubscriptionLength - 1
+				&& InSubscription[SubscriptionPos] == FMQTTTopic::SingleLevelWildcardChar)
+			{
+				// Subscription wildcard preceded by level separator
+				return SubscriptionPos > 0 && InSubscription[SubscriptionPos - 1] == FMQTTTopic::TopicLevelSeparatorChar;
+			}
+		}
+		// Char mismatch
+		else
+		{
+			// SingleLevel Wildcard
+			if (InSubscription[SubscriptionPos] == FMQTTTopic::SingleLevelWildcardChar)
+			{
+				// ...must be enclosed in level separator
+				if ((SubscriptionPos > 0 && InSubscription[SubscriptionPos - 1] != FMQTTTopic::TopicLevelSeparatorChar)
+					|| (SubscriptionPos < SubscriptionLength - 1 && InSubscription[SubscriptionPos + 1] != FMQTTTopic::TopicLevelSeparatorChar))
+				{
+					return false;
+				}
+
+				// Proceed until level separator found
+				SubscriptionPos++;
+				while (TopicPos < TopicLength && InTopic[TopicPos] != FMQTTTopic::TopicLevelSeparatorChar)
+				{
+					TopicPos++;
+				}
+
+				// Both topic and subscription end in level separator
+				if(TopicPos == TopicLength && SubscriptionPos == SubscriptionLength)
+				{
+					return true;
+				}
+			}
+			// MultiLevel Wildcard
+			else if (InSubscription[SubscriptionPos] == FMQTTTopic::MultiLevelWildcardChar)
+			{
+				// ...must be preceded by level separator
+				if (SubscriptionPos > 0
+					&& InSubscription[SubscriptionPos - 1] != FMQTTTopic::TopicLevelSeparatorChar)
+				{
+					return false;
+				}
+
+				// ...and is the last character
+				return SubscriptionPos + 1 == SubscriptionLength;
+			}
+			else
+			{
+				// Previous character was SingleLevel wildcard, remaining subscription is '/#'
+				return SubscriptionPos > 0
+					&& SubscriptionPos + 2 == SubscriptionLength
+					&& TopicPos == TopicLength
+					&& InSubscription[SubscriptionPos - 1] == FMQTTTopic::SingleLevelWildcardChar
+					&& InSubscription[SubscriptionPos] == FMQTTTopic::TopicLevelSeparatorChar
+					&& InSubscription[SubscriptionPos + 1] == FMQTTTopic::MultiLevelWildcardChar;
+			}
+		}
+	}
+
+	return false;
+}
+
+bool FMQTTSubscription::Matches(const FString& InTopic) const
+{
+	return FMQTTSubscription::Matches(InTopic, Topic.ToString());
+}
+
+bool FMQTTSubscription::IsValid() const
+{
+	return false;
+}
+
+#undef LOCTEXT_NAMESPACE // "MQTT"
diff --git a/Engine/Plugins/Protocols/MQTT/Source/MQTTCore/Private/MQTTSharedInternal.cpp b/Engine/Plugins/Protocols/MQTT/Source/MQTTCore/Private/MQTTSharedInternal.cpp
new file mode 100644
--- /dev/null
+++ b/Engine/Plugins/Protocols/MQTT/Source/MQTTCore/Private/MQTTSharedInternal.cpp
@@ -0,0 +1,74 @@
+﻿// Copyright Epic Games, Inc. All Rights Reserved.
+
+#include "MQTTSharedInternal.h"
+
+namespace Internal
+{
+	static TMap<ESocketConnectionState, const TCHAR*> SocketConnectionNames =
+	{
+		{ ESocketConnectionState::SCS_Connected, TEXT("Connected") },
+		{ ESocketConnectionState::SCS_ConnectionError, TEXT("Connection Error") },
+		{ ESocketConnectionState::SCS_NotConnected, TEXT("Not Connected") }
+	};
+			
+	static TMap<EMQTTConnectReturnCode, const TCHAR*> ConnectReturnCodeDescriptions =
+	{
+		{ EMQTTConnectReturnCode::Accepted, TEXT("Connection accepted.") },
+		{ EMQTTConnectReturnCode::RefusedProtocolVersion, TEXT("The Server does not the support the level of the MQTT protocol requested by the Client.") },
+		{ EMQTTConnectReturnCode::RefusedIdentifierRejected, TEXT("The Client identifier is correct UTF-8 but not allowed by the Server.") },
+		{ EMQTTConnectReturnCode::RefusedServerUnavailable, TEXT("The Network Connection has been made but the MQTT service is unavailable.") },
+		{ EMQTTConnectReturnCode::RefusedBadUsernamePassword, TEXT("The data in the user name or password is malformed.") },
+		{ EMQTTConnectReturnCode::RefusedNotAuthorized, TEXT("The Client is not authorized to connect.") },
+		{ EMQTTConnectReturnCode::AlreadyConnected, TEXT("The Client is already connected to the Server.") },
+		{ EMQTTConnectReturnCode::InvalidURL, TEXT("The provided URL is malformed.") },
+		{ EMQTTConnectReturnCode::SocketError, TEXT("There was a socket error.") }
+	};
+
+	static TMap<EMQTTSocketState, const TCHAR*> SocketStateNames =
+	{
+		{ EMQTTSocketState::Connecting, TEXT("Connecting") },
+		{ EMQTTSocketState::Connected, TEXT("Connected") },
+		{ EMQTTSocketState::Error, TEXT("Error") },
+		{ EMQTTSocketState::Stopping, TEXT("Stopping") },
+		{ EMQTTSocketState::ConnectionRejected, TEXT("ConnectionRejected") },
+	};
+	
+	static TMap<EMQTTState, const TCHAR*> StateNames =
+	{
+		{ EMQTTState::Unknown, TEXT("Unknown") },
+		{ EMQTTState::Connected, TEXT("Connected") },
+		{ EMQTTState::Connecting, TEXT("Connecting") },
+		{ EMQTTState::Disconnected, TEXT("Disconnected") },
+		{ EMQTTState::Disconnecting, TEXT("Disconnecting") },
+		{ EMQTTState::Stopping, TEXT("Stopping") },
+		{ EMQTTState::None, TEXT("None") }
+	};
+
+	static TMap<EMQTTSubscribeReturnCode, const TCHAR*> SubscribeReturnCodeNames =
+	{
+		{ EMQTTSubscribeReturnCode::Success_QoS0, TEXT("Quality of Service 0") },
+		{ EMQTTSubscribeReturnCode::Success_QoS1, TEXT("Quality of Service 1") },
+		{ EMQTTSubscribeReturnCode::Success_QoS2, TEXT("Quality of Service 2") },
+		{ EMQTTSubscribeReturnCode::Failure, TEXT("Failure") },
+	};
+}
+
+const TCHAR* MQTT::GetSocketConnectionStateName(ESocketConnectionState InConnectionState)
+{
+	return Internal::SocketConnectionNames[InConnectionState];
+}
+
+const TCHAR* MQTT::GetMQTTStateName(const EMQTTState InState)
+{
+	return Internal::StateNames[InState];
+}
+
+const TCHAR* MQTT::GetMQTTConnectReturnCodeDescription(const EMQTTConnectReturnCode InConnectReturnCode)
+{
+	return Internal::ConnectReturnCodeDescriptions[InConnectReturnCode];
+}
+
+const TCHAR* MQTT::GetMQTTSubscribeReturnCodeName(const EMQTTSubscribeReturnCode InSubscribeReturnCode)
+{
+	return Internal::SubscribeReturnCodeNames[InSubscribeReturnCode];
+}
diff --git a/Engine/Plugins/Protocols/MQTT/Source/MQTTCore/Private/MQTTSharedInternal.h b/Engine/Plugins/Protocols/MQTT/Source/MQTTCore/Private/MQTTSharedInternal.h
new file mode 100644
--- /dev/null
+++ b/Engine/Plugins/Protocols/MQTT/Source/MQTTCore/Private/MQTTSharedInternal.h
@@ -0,0 +1,119 @@
+﻿// Copyright Epic Games, Inc. All Rights Reserved..
+
+#pragma once
+
+#include "CoreMinimal.h"
+#include "MQTTProtocol.h"
+#include "MQTTShared.h"
+#include "Serialization/BufferArchive.h"
+
+#define REGISTER_MQTT_ARCHIVE(MQTTPacket)											\
+FArchive& operator<<(FArchive & Ar, MQTTPacket &Packet)								\
+{																					\
+	Packet.Serialize(Ar);															\
+	return Ar;																		\
+}
+
+#define REGISTER_MQTT_PACK(MQTTPacket)												\
+bool MQTTPacket::Pack(FBufferArchive& InArchive)									\
+{																					\
+	InArchive << *this;																\
+	return true;																	\
+}
+
+// #include "MQTTSharedInternal.generated.h"
+
+/** Connection socket state */
+enum class EMQTTSocketState : uint8
+{
+	Connecting = 0,
+	Connected = 1,			// SCS_Connected
+	Error = 2,				// Couldn't connect / SCS_ConnectionError
+	ConnectionRejected = 3,	// Could connect, but was refused or rejected
+	Stopping = 8,			// Force stop/shutdown
+};
+
+/** MQTT protocol state */
+enum class EMQTTState : uint8
+{
+	Unknown = 0,
+	Connected = 1,
+	Connecting = 2,
+	Disconnected = 3,
+	Disconnecting = 4,
+	Stopping = 8,			// Force stop/shutdown
+
+	None = 127				// Used for state dependency to indicate there is none
+};
+
+namespace MQTT
+{
+	const TCHAR* GetSocketConnectionStateName(ESocketConnectionState InConnectionState);
+	const TCHAR* GetMQTTPacketTypeName(EMQTTPacketType InMessageType);
+	const TCHAR* GetMQTTStateName(EMQTTState InState);
+	const TCHAR* GetMQTTConnectReturnCodeDescription(EMQTTConnectReturnCode InConnectReturnCode);
+	const TCHAR* GetMQTTSubscribeReturnCodeName(EMQTTSubscribeReturnCode InSubscribeReturnCode);	
+};
+
+/** An arbitrary task allowing the storage/calls and ResultType to be separate (vs. a raw TPromise). */
+class IMQTTTask
+{
+public:
+	virtual ~IMQTTTask() = default;
+	virtual void Execute() = 0;
+	virtual void Abandon() = 0;
+
+	EMQTTState StateDependency = EMQTTState::None;
+};
+
+template <typename ResultType>
+class TMQTTTask final : public IMQTTTask
+{
+public:
+	TMQTTTask(TUniqueFunction<ResultType()>&& InFunction, const EMQTTState& InStateDependency = EMQTTState::Connected)
+		: Function(MoveTemp(InFunction))
+	{
+		StateDependency = InStateDependency;
+	}
+
+	virtual ~TMQTTTask() override
+	{
+	}
+
+	virtual void Execute() override
+	{
+		Function();
+	}
+
+	virtual void Abandon() override
+	{
+		// Just complete the task
+		Function();
+	}
+
+private:
+	friend class FMQTTClientRunnable;
+
+	/** The function to execute. */
+	TUniqueFunction<ResultType()> Function; 
+};
+
+namespace MQTT
+{
+	template <typename CallableType>
+	auto CreateTask(CallableType&& InCallable, const EMQTTState& InStateDependency) -> TSharedPtr<TMQTTTask<decltype(Forward<CallableType>(InCallable)())>>
+	{
+		using ResultType = decltype(Forward<CallableType>(InCallable)());
+		TUniqueFunction<ResultType()> Function(Forward<CallableType>(InCallable));
+		return MakeShared<TMQTTTask<ResultType>>(MoveTemp(Function), InStateDependency);
+	}
+
+	// template <typename RequestType, typename ResponseType>
+	// auto CreateMessage(RequestType&& InRequest) -> TSharedPtr<TMQTTRequestResponseOperation<RequestType, ResponseType>>
+	// {
+	// 	return MakeShared<>()
+	// 	using ResultType = decltype(Forward<CallableType>(InCallable)());
+	// 	TUniqueFunction<ResultType()> Function(Forward<CallableType>(InCallable));
+	// 	return MakeShared<TMQTTTask<ResultType>>(MoveTemp(Function), InStateDependency);
+	// }
+}
diff --git a/Engine/Plugins/Protocols/MQTT/Source/MQTTCore/Private/MQTTSubsystem.cpp b/Engine/Plugins/Protocols/MQTT/Source/MQTTCore/Private/MQTTSubsystem.cpp
new file mode 100644
--- /dev/null
+++ b/Engine/Plugins/Protocols/MQTT/Source/MQTTCore/Private/MQTTSubsystem.cpp
@@ -0,0 +1,55 @@
+﻿// Copyright Epic Games, Inc. All Rights Reserved.
+
+#include "MQTTSubsystem.h"
+
+#include "JsonObjectWrapper.h"
+#include "MQTTClientObject.h"
+#include "MQTTClientSettings.h"
+#include "Engine/Engine.h"
+
+UMQTTClientObject* UMQTTSubsystem::GetOrCreateClient_WithProjectURL(UObject* InParent)
+{
+	const UMQTTClientSettings* Settings = GetDefault<UMQTTClientSettings>();
+	return GetOrCreateClient(InParent, Settings->DefaultURL);
+}
+
+UMQTTClientObject* UMQTTSubsystem::GetOrCreateClient(UObject* InParent, const FMQTTURL& InURL)
+{
+	const FGuid ClientId = InURL.ToGuid();
+	
+	UMQTTSubsystem* Self = GEngine->GetEngineSubsystem<UMQTTSubsystem>();
+	if(Self)
+	{
+		UMQTTClientObject* MQTTClientObject = Self->Clients.FindOrAdd(ClientId, NewObject<UMQTTClientObject>(InParent));
+
+		// Might exist but be destroyed
+		if(!MQTTClientObject || !IsValid(MQTTClientObject))
+		{
+			Self->Clients[ClientId] = MQTTClientObject = NewObject<UMQTTClientObject>(InParent);			
+		}
+
+		ensure(IsValid(MQTTClientObject));
+		
+		MQTTClientObject->Initialize(InURL);
+		return MQTTClientObject;
+	}
+
+	return nullptr;
+}
+
+FString UMQTTSubsystem::GetPayloadString(const FMQTTClientMessage& InClientMessage)
+{
+	return InClientMessage.GetPayloadAsString();
+}
+
+bool UMQTTSubsystem::GetPayloadJson(UObject* InParent, const FMQTTClientMessage& InClientMessage, FJsonObjectWrapper& OutJson)
+{
+	// @todo: fix
+	// TSharedPtr<FJsonObject> JsonObject;
+	// if(InClientMessage.GetPayloadAsJson(JsonObject))
+	// {
+	// 	return UJsonObject::FromNative(InParent, JsonObject);
+	// }
+
+	return false;
+}
diff --git a/Engine/Plugins/Protocols/MQTT/Source/MQTTCore/Public/IMQTTClient.h b/Engine/Plugins/Protocols/MQTT/Source/MQTTCore/Public/IMQTTClient.h
new file mode 100644
--- /dev/null
+++ b/Engine/Plugins/Protocols/MQTT/Source/MQTTCore/Public/IMQTTClient.h
@@ -0,0 +1,101 @@
+﻿// Copyright Epic Games, Inc. All Rights Reserved.
+
+#pragma once
+
+#include "CoreMinimal.h"
+#include "MQTTProtocol.h"
+
+#include "MQTTShared.h"
+
+struct FMQTTURL;
+struct FMQTTClientMessage;
+struct FMQTTClientTask;
+struct FMQTTSubscription;
+
+class IMQTTClient
+{
+public:
+	DECLARE_MULTICAST_DELEGATE_OneParam(FOnConnect, EMQTTConnectReturnCode)
+	virtual FOnConnect& OnConnect() = 0;
+
+	DECLARE_MULTICAST_DELEGATE(FOnDisconnect)
+	virtual FOnDisconnect& OnDisconnect() = 0;
+
+	DECLARE_MULTICAST_DELEGATE(FOnPublish)
+	virtual FOnPublish& OnPublish() = 0;
+
+	DECLARE_MULTICAST_DELEGATE_OneParam(FOnSubscribe, TArray<FMQTTSubscribeResult> /* SubscriptionHandles */)
+	virtual FOnSubscribe& OnSubscribe() = 0;
+
+	DECLARE_MULTICAST_DELEGATE(FOnUnsubscribe)
+	virtual FOnUnsubscribe& OnUnsubscribe() = 0;
+
+	DECLARE_MULTICAST_DELEGATE_OneParam(FOnMessage, const FMQTTClientMessage& /* Packet */)
+	virtual FOnMessage& OnMessage() = 0;	
+
+public:
+	virtual ~IMQTTClient() = default;
+
+	// @todo: store session
+	/** Connect using the clients URL. Set bCleanSession = false to resume previous session. */
+	virtual TFuture<EMQTTConnectReturnCode> Connect(bool bCleanSession = true) = 0;
+	
+	//virtual TCancellablePromise<bool> Connect();
+	virtual TFuture<void> Disconnect() = 0;
+
+	/** Return true if the entire operation was successful (depends on QoS) */
+	virtual TFuture<bool> Publish(const FString& InTopic, const TArray<uint8>& InPayload, EMQTTQualityOfService InQoS = EMQTTQualityOfService::Once, const bool bInRetain = false) = 0;
+
+	/** Return true if the entire operation was successful (depends on QoS) */
+	virtual TFuture<bool> Publish(const FString& InTopic, const FString& InPayload, EMQTTQualityOfService InQoS = EMQTTQualityOfService::Once, const bool bInRetain = false) = 0;
+
+	virtual TFuture<TArray<FMQTTSubscribeResult>> Subscribe(const TArray<TPair<FString, EMQTTQualityOfService>>& InTopicFilterQoSPairs) = 0;
+
+	// only for single sub
+	template <typename CallableType>
+	TFuture<FMQTTSubscribeResult> Subscribe(const FString& InTopicFilter, CallableType&& InOnMessage, EMQTTQualityOfService InQoS = EMQTTQualityOfService::Once);
+
+	virtual TFuture<bool> Unsubscribe(const TSet<FString>& InTopicFilters) = 0;
+
+	/** Returns true if response received before timeout. */
+	virtual TFuture<bool> Ping(const float& InTimeout = 2.0f) = 0;
+
+	/** Unique Id for this client. */
+	virtual FGuid GetClientId() const = 0;
+
+	/** URL for this client. */
+	virtual const FMQTTURL& GetURL() const = 0;
+
+	/** Is client currently connected? */
+	virtual bool IsConnected() const = 0;
+	
+	/** Validity of this client. */
+	virtual bool IsValid() const = 0;
+};
+
+template <typename CallableType>
+TFuture<FMQTTSubscribeResult> IMQTTClient::Subscribe(const FString& InTopicFilter, CallableType&& InOnMessage, EMQTTQualityOfService InQoS)
+{
+	return Subscribe({MakeTuple(InTopicFilter, InQoS)})
+		.Next([&](const TArray<FMQTTSubscribeResult>& InSubscribeResult)
+		{
+			if(InSubscribeResult.Num() == 0)
+			{
+				return FMQTTSubscribeResult(EMQTTSubscribeReturnCode::Failure, {});
+			}
+
+			if(InSubscribeResult[0].ReturnCode == EMQTTSubscribeReturnCode::Failure)
+			{
+				return InSubscribeResult[0];				
+			}
+			
+			// Automatically subscribe to message using user-supplied function ...
+			InSubscribeResult[0].Subscription->OnSubscriptionMessage().AddLambda(MoveTemp(InOnMessage));
+			return InSubscribeResult[0];
+		});
+}
+
+FORCEINLINE uint32 GetTypeHash(const IMQTTClient& InClient)
+{
+	return GetTypeHash(InClient.GetClientId());
+}
diff --git a/Engine/Plugins/Protocols/MQTT/Source/MQTTCore/Public/IMQTTCoreModule.h b/Engine/Plugins/Protocols/MQTT/Source/MQTTCore/Public/IMQTTCoreModule.h
new file mode 100644
--- /dev/null
+++ b/Engine/Plugins/Protocols/MQTT/Source/MQTTCore/Public/IMQTTCoreModule.h
@@ -0,0 +1,42 @@
+﻿// Copyright Epic Games, Inc. All Rights Reserved.
+
+#pragma  once
+
+#include "CoreMinimal.h"
+#include "Modules/ModuleInterface.h"
+#include "Modules/ModuleManager.h"
+
+struct FMQTTURL;
+class IMQTTServer;
+class IMQTTClient;
+
+class IMQTTCoreModule : public IModuleInterface
+{
+public:
+	/**
+	* Singleton-like access to this module's interface.  This is just for convenience!
+	* Beware of calling this during the shutdown phase, though.  Your module might have been unloaded already.
+	*
+	* @return Returns singleton instance, loading the module on demand if needed
+	*/
+	static IMQTTCoreModule& Get()
+	{
+		static const FName ModuleName = "MQTTCore";
+		return FModuleManager::LoadModuleChecked<IMQTTCoreModule>(ModuleName);
+	}
+
+	/** Get or creates an MQTT client that uses the project default URL. An existing client is returned for a matching URL. */
+	virtual TSharedPtr<IMQTTClient, ESPMode::ThreadSafe> GetOrCreateClient(bool bForceNew = false) = 0;
+
+	/** Creates an MQTT client from a URL. An existing client is returned for a matching URL. bForceNew = recreate if existing */
+	virtual TSharedPtr<IMQTTClient, ESPMode::ThreadSafe> GetOrCreateClient(const FMQTTURL& InURL, bool bForceNew = false) = 0;
+
+protected:
+#if WITH_EDITOR
+#if WITH_DEV_AUTOMATION_TESTS
+	friend class FMQTTClientSpec;
+#endif
+#endif
+
+	virtual int32 GetClientNum() const = 0;
+};
diff --git a/Engine/Plugins/Protocols/MQTT/Source/MQTTCore/Public/MQTTClientMessage.h b/Engine/Plugins/Protocols/MQTT/Source/MQTTCore/Public/MQTTClientMessage.h
new file mode 100644
--- /dev/null
+++ b/Engine/Plugins/Protocols/MQTT/Source/MQTTCore/Public/MQTTClientMessage.h
@@ -0,0 +1,55 @@
+﻿// Copyright Epic Games, Inc. All Rights Reserved.
+
+#pragma once
+
+#include "CoreMinimal.h"
+#include "MQTTShared.h"
+#include "Dom/JsonObject.h"
+
+#include "MQTTClientMessage.generated.h"
+
+USTRUCT(BlueprintType)
+struct MQTTCORE_API FMQTTClientMessage
+{
+	GENERATED_BODY()
+
+	/** TimeStamp as UTC. */
+	UPROPERTY(EditAnywhere, BlueprintReadWrite, Category = "MQTT")
+	FDateTime TimeStamp = FDateTime::UtcNow();
+
+	/** Packet topic. */
+	UPROPERTY(EditAnywhere, BlueprintReadWrite, Category = "MQTT")
+	FString Topic;
+
+	// @todo: func lib to get as string
+	/** Packet content. */
+	UPROPERTY(EditAnywhere, BlueprintReadWrite, Category = "MQTT")
+	TArray<uint8> Payload;
+
+	/** Retain flag. */
+	UPROPERTY(EditAnywhere, BlueprintReadWrite, Category = "MQTT")
+	bool bRetain = false;
+
+	/** Quality of Service. */
+	UPROPERTY(EditAnywhere, BlueprintReadWrite, Category = "MQTT", meta = (DisplayName = "Quality of Service"))
+	EMQTTQualityOfService QoS = EMQTTQualityOfService::Once;
+
+	/** Get's the payload as a string, ie. for Json. */
+	const FString& GetPayloadAsString() const { return PayloadString; }
+
+	/** Get's the payload as a Json object, first converting to string. Returns false if unsuccessful. */
+	bool GetPayloadAsJson(TSharedPtr<FJsonObject>& OutJson) const;
+
+	/** Set's the payload from the input string. */
+	void SetPayloadFromString(const FString& InPayloadString);
+
+	/** Gets the validity of the Packet. */
+	bool IsValid() const { return true; }
+
+protected:
+	friend class FMQTTClient;
+	
+	/** Cached payload-as-string, ie. Json. */
+	UPROPERTY()
+	FString PayloadString;
+};
diff --git a/Engine/Plugins/Protocols/MQTT/Source/MQTTCore/Public/MQTTClientObject.h b/Engine/Plugins/Protocols/MQTT/Source/MQTTCore/Public/MQTTClientObject.h
new file mode 100644
--- /dev/null
+++ b/Engine/Plugins/Protocols/MQTT/Source/MQTTCore/Public/MQTTClientObject.h
@@ -0,0 +1,129 @@
+﻿// Copyright Epic Games, Inc. All Rights Reserved.
+
+#pragma once
+
+#include "CoreMinimal.h"
+#include "MQTTClientMessage.h"
+#include "MQTTShared.h"
+
+#include "MQTTClientObject.generated.h"
+
+class IMQTTClient;
+
+UCLASS(BlueprintType, meta = (DisplayName = "MQTT Subscription"))
+class UMQTTSubscriptionObject final : public UObject
+{
+	GENERATED_BODY()
+	
+public:
+	virtual ~UMQTTSubscriptionObject() override;
+	
+	DECLARE_DYNAMIC_DELEGATE_OneParam(FOnMessageDelegate, const FMQTTClientMessage&, Message);
+
+	/** Initialize from C++ subscription */
+	void Initialize(const TSharedPtr<FMQTTSubscription, ESPMode::ThreadSafe>& InSubscription);
+
+	UFUNCTION(BlueprintCallable, Category = "MQTT")
+	virtual void SetOnMessageHandler(const FOnMessageDelegate& OnMessageCallback);
+
+	UFUNCTION(BlueprintCallable, BlueprintPure, Category = "MQTT")
+	bool IsValid() const { return MqttSubscriptionPtr.IsValid(); }
+
+	friend uint32 GetTypeHash(const UMQTTSubscriptionObject* InSubscription)
+	{
+		if(!::IsValid(InSubscription) || !InSubscription->IsValid())
+		{
+			return 0;
+		}
+
+		return GetTypeHash(*InSubscription->MqttSubscriptionPtr.Get());
+	}
+
+protected:
+	UPROPERTY()
+	FOnMessageDelegate OnMessageDelegate;
+
+private:
+	TSharedPtr<FMQTTSubscription, ESPMode::ThreadSafe> MqttSubscriptionPtr;
+
+	void OnMessage(const FMQTTClientMessage& InMessage) const;
+};
+
+UCLASS(BlueprintType, meta = (DisplayName = "MQTT Client"))
+class UMQTTClientObject final : public UObject
+{	
+	GENERATED_BODY()
+	
+public:
+	virtual ~UMQTTClientObject() override;
+
+	DECLARE_DYNAMIC_DELEGATE_OneParam(FOnConnectDelegate, EMQTTConnectReturnCode, ReturnCode);
+	DECLARE_DYNAMIC_DELEGATE(FOnDisconnectDelegate);
+	DECLARE_DYNAMIC_DELEGATE(FOnPublishDelegate);	
+	DECLARE_DYNAMIC_DELEGATE_OneParam(FOnSubscribeDelegate, EMQTTSubscribeReturnCode, ReturnCode);
+	DECLARE_DYNAMIC_DELEGATE(FOnUnsubscribeDelegate);
+	DECLARE_DYNAMIC_DELEGATE_OneParam(FOnMessageDelegate, const FMQTTClientMessage&, Message);
+
+	/** Initialize from URL */
+	void Initialize(const FMQTTURL& InURL);
+	
+	void InitDelegates();
+	void RemoveDelegates() const;
+
+	virtual void BeginDestroy() override;
+
+	TSharedPtr<IMQTTClient, ESPMode::ThreadSafe> GetClientImpl() const;
+
+	UFUNCTION(BlueprintCallable, Category = "MQTT")
+	void Connect(UPARAM(DisplayName = "OnConnect") const FOnConnectDelegate& InOnConnect);
+
+	UFUNCTION(BlueprintCallable, Category = "MQTT")
+	void Disconnect(UPARAM(DisplayName = "OnDisconnect") const FOnDisconnectDelegate& InOnDisconnect);
+
+	/** OutMessageId can be used to match this request with the callback response. */
+	UFUNCTION(BlueprintCallable, Category = "MQTT")
+	void Publish(
+		UPARAM(DisplayName = "Topic") const FString& InTopic,
+		UPARAM(DisplayName = "Payload") const TArray<uint8>& InPayload,
+		UPARAM(DisplayName = "Quality of Service") EMQTTQualityOfService InQoS = EMQTTQualityOfService::Once,
+		const bool bInRetain = false);
+
+	/** OutMessageId can be used to match this request with the callback response. */
+	UFUNCTION(BlueprintCallable, Category = "MQTT")
+	UMQTTSubscriptionObject* Subscribe(
+		UPARAM(DisplayName = "Topic") const FString& InTopic,
+		UPARAM(DisplayName = "Quality of Service") EMQTTQualityOfService InQoS = EMQTTQualityOfService::Once);
+
+	/** OutMessageId can be used to match this request with the callback response. */
+	UFUNCTION(BlueprintCallable, Category = "MQTT", meta = (DisplayName = "Subscribe (Multiple Topics)"))
+	TArray<UMQTTSubscriptionObject*> SubscribeMany(
+		UPARAM(DisplayName = "Topics") const TArray<FString>& InTopics,
+		UPARAM(DisplayName = "Quality of Services") const TArray<EMQTTQualityOfService> InQoS);
+
+	/** OutMessageId can be used to match this request with the callback response. */
+	UFUNCTION(BlueprintCallable, Category = "MQTT")
+	void Unsubscribe(UPARAM(DisplayName = "Topic") const FString& InTopic);
+
+	UFUNCTION(BlueprintPure, Category = "MQTT")
+	FString GetClientId() const;
+
+	/** URL for this client. */
+	UFUNCTION(BlueprintPure, Category = "MQTT")
+	FMQTTURL GetURL() const;
+
+protected:
+	UPROPERTY()
+	FOnConnectDelegate OnConnectDelegate;
+	
+	UPROPERTY()
+	FOnDisconnectDelegate OnDisconnectDelegate;
+
+	UPROPERTY()
+	FOnMessageDelegate OnMessageDelegate;
+
+	UPROPERTY()
+	TArray<UMQTTSubscriptionObject*> Subscriptions; 
+
+private:
+	TSharedPtr<IMQTTClient, ESPMode::ThreadSafe> MqttClientPtr;
+};
diff --git a/Engine/Plugins/Protocols/MQTT/Source/MQTTCore/Public/MQTTClientSettings.h b/Engine/Plugins/Protocols/MQTT/Source/MQTTCore/Public/MQTTClientSettings.h
new file mode 100644
--- /dev/null
+++ b/Engine/Plugins/Protocols/MQTT/Source/MQTTCore/Public/MQTTClientSettings.h
@@ -0,0 +1,26 @@
+// Copyright Epic Games, Inc. All Rights Reserved.
+
+#pragma once
+
+#include "CoreMinimal.h"
+
+#include "MQTTShared.h"
+
+#include "MQTTClientSettings.generated.h"
+
+/**
+* MQTT Client Settings
+*/
+UCLASS(Config = MQTT)
+class MQTTCORE_API UMQTTClientSettings final : public UObject
+{
+	GENERATED_BODY()
+
+public:
+	UPROPERTY(EditAnywhere, Config, Category = "Connection")
+	FMQTTURL DefaultURL;
+
+	/** Maximum messages to publish per second. */
+	UPROPERTY(EditAnywhere, Config, Category = "Bandwidth")
+	uint32 PublishRate = 60;
+};
diff --git a/Engine/Plugins/Protocols/MQTT/Source/MQTTCore/Public/MQTTProtocol.h b/Engine/Plugins/Protocols/MQTT/Source/MQTTCore/Public/MQTTProtocol.h
new file mode 100644
--- /dev/null
+++ b/Engine/Plugins/Protocols/MQTT/Source/MQTTCore/Public/MQTTProtocol.h
@@ -0,0 +1,197 @@
+﻿// Copyright Epic Games, Inc. All Rights Reserved.
+
+// Types that correspond to the MQTT specification v3.1.1 and v5.0
+
+#pragma once
+
+#include "CoreMinimal.h"
+#include "UObject/Object.h"
+
+#include "MQTTProtocol.generated.h"
+
+// @todo: max keepalive is 18h 12m 15s - MQTT spec v3.1.1 s 3.1.2.10
+
+class FMQTTProtocol
+{
+};
+
+#define PROTOCOL_NAME "MQTT"
+
+#define PROTOCOL_VERSION_v311 4
+#define PROTOCOL_VERSION_v5 5
+
+enum class EMQTTProtocolLevel : uint8
+{
+	MQTT311 = 4,
+	MQTT5 = 5
+};
+
+/* Control packet types */
+enum class EMQTTPacketType : uint8
+{
+	Connect = 0x10U,
+	ConnectAcknowledge = 0x20U,
+
+	Publish = 0x30U,
+	PublishAcknowledge = 0x40U,
+	PublishReceived = 0x50U,
+	PublishRelease = 0x60U,
+	PublishComplete = 0x70U,
+
+	Subscribe = 0x80U,
+	SubscribeAcknowledge = 0x90U,
+
+	Unsubscribe = 0xA0U,
+	UnsubscribeAcknowledge = 0xB0U,
+
+	PingRequest = 0xC0U,
+	PingResponse = 0xD0U,
+
+	Disconnect = 0xE0U,
+	Authorize = 0xF0U,
+
+	/** non-spec entries */
+	None = 0,
+};
+
+/** MQTT v3.1.1 */
+UENUM(BlueprintType)
+enum class EMQTTConnectReturnCode : uint8
+{
+    /** Connection accepted. */
+	Accepted = 0x00U,
+    /** The Server does not the support the level of the MQTT protocol requested by the Client. */
+	RefusedProtocolVersion = 0x01U,
+    /** The Client identifier is correct UTF-8 but not allowed by the Server. */
+	RefusedIdentifierRejected = 0x02U,
+    /** The Network Connection has been made but the MQTT service is unavailable. */
+	RefusedServerUnavailable = 0x03U,
+    /** The data in the user name or password is malformed. */
+	RefusedBadUsernamePassword = 0x04U,
+    /** The Client is not authorized to connect. */
+	RefusedNotAuthorized = 0x05U,
+
+	/** Non-spec entries */
+    /** The Client is already connected to the Server. */
+	AlreadyConnected,
+	/** The provided URL is malformed. */
+	InvalidURL,
+	/** Socket error. */
+	SocketError
+};
+
+/** MQTT v3.1.1 */
+UENUM(BlueprintType)
+enum class EMQTTSubscribeReturnCode : uint8
+{
+    Success_QoS0 = 0x00U,           /* SubscribeAcknowledge */
+    Success_QoS1 = 0x01U,           /* SubscribeAcknowledge */
+    Success_QoS2 = 0x02U,           /* SubscribeAcknowledge */
+    Failure = 0x80U                 /* SubscribeAcknowledge */
+};
+
+/** MQTT v5 (renamed from ReturnCode in v3.1.1) */
+// @note: conflicting codes a differentiated by what they're returned from
+enum class EMQTTReasonCode : uint8
+{
+	Success = 0x00U,						        /* ConnectAcknowledge, PublishAcknowledge, PublishReceive, PublishRelease, PublishComplete, UnsubscribeAcknowledge, Authorize */
+	NormalDisconnection = 0x00U,			        /* Disconnect */
+	Granted_QoS0 = 0x00U,					        /* SubscribeAcknowledge */
+	Granted_QoS1 = 0x01U,					        /* SubscribeAcknowledge */
+	Granted_QoS2 = 0x02U,					        /* SubscribeAcknowledge */
+	DisconnectWithWillMsg = 0x04U,		            /* Disconnect */
+	NoMatchingSubscribers = 0x10U,		            /* PublishAcknowledge, PublishReceive */
+	NoSubscriptionExisted = 0x11U,		            /* UnsubscribeAcknowledge */
+	ContinueAuthentication = 0x18U,					/* Authorize */
+	ReAuthenticate = 0x19U,							/* Authorize */
+
+	Unspecified = 0x80U,					        /* ConnectAcknowledge, PublishAcknowledge, PublishReceive, SubscribeAcknowledge, UnsubscribeAcknowledge, Disconnect */
+	MalformedPacket = 0x19U,				        /* ConnectAcknowledge, Disconnect */
+	ProtocolError = 0x82U,							/* Disconnect */
+	ImplementationSpecific = 0x83U,					/* ConnectAcknowledge, PublishAcknowledge, PublishReceive, SubscribeAcknowledge, UnsubscribeAcknowledge, Disconnect */
+	UnsupportedProtocolVersion = 0x84U,				/* ConnectAcknowledge */
+	ClientIdNotValid = 0x85U,			            /* ConnectAcknowledge */
+	BadUsernameOrPassword = 0x86U,					/* ConnectAcknowledge */
+	NotAuthorized = 0x87U,							/* ConnectAcknowledge, PublishAcknowledge, PublishReceive, SubscribeAcknowledge, UnsubscribeAcknowledge, Disconnect */
+	ServerUnavailable = 0x88U,						/* ConnectAcknowledge */
+	ServerBusy = 0x89U,								/* ConnectAcknowledge, Disconnect */
+	Banned = 0x8AU,									/* ConnectAcknowledge */
+	ServerShuttingDown = 0x8BU,						/* Disconnect */
+	BadAuthenticationMethod = 0x8CU,				/* ConnectAcknowledge */
+	KeepAliveTimeout = 0x8DU,			            /* Disconnect */
+	SessionTakenOver = 0x8EU,			            /* Disconnect */
+	TopicFilterInvalid = 0x8FU,						/* SubscribeAcknowledge, UnsubscribeAcknowledge, Disconnect */
+	TopicNameInvalid = 0x90U,			            /* ConnectAcknowledge, PublishAcknowledge, PublishReceive, Disconnect */
+	PacketIdInUse = 0x91U,							/* PublishAcknowledge, SubscribeAcknowledge, UnsubscribeAcknowledge */
+	PacketIdNotFound = 0x92U,			            /* PublishRelease, PublishComplete */
+	ReceiveMaximumExceeded = 0x93U,					/* Disconnect */
+	TopicAliasInvalid = 0x94U,						/* Disconnect */
+	PacketTooLarge = 0x95U,							/* ConnectAcknowledge, PublishAcknowledge, PublishReceive, Disconnect */
+	MessageRateTooHigh = 0x96U,						/* Disconnect */
+	QuotaExceeded = 0x97U,							/* PublishAcknowledge, PublishReceive, SubscribeAcknowledge, Disconnect */
+	AdministrativeAction = 0x98U,		            /* Disconnect */
+	PayloadFormatInvalid = 0x99U,		            /* ConnectAcknowledge, Disconnect */
+	RetainNotSupported = 0x9AU,						/* ConnectAcknowledge, Disconnect */
+	QoSNotSupported = 0x9BU,			            /* ConnectAcknowledge, Disconnect */
+	UseAnotherServer = 0x9CU,			            /* ConnectAcknowledge, Disconnect */
+	ServerMoved = 0x9DU,					        /* ConnectAcknowledge, Disconnect */
+	SharedSubscriptionNotSupported = 0x9EU,			/* SubscribeAcknowledge, Disconnect */
+	ConnectionRateExceeded = 0x9FU,					/* ConnectAcknowledge, Disconnect */
+	MaximumConnectTime = 0xA0U,						/* Disconnect */
+	SubscriptionIdsNotSupported = 0xA1U,			/* SubscribeAcknowledge, Disconnect */
+	WildcardSubscriptionNotSupported = 0xA2U,	    /* SubscribeAcknowledge, Disconnect */
+};
+
+enum class FMQTT5Property: uint8
+{
+	PayloadFormatIndicator = 1,		    /* Byte :				Publish, Will Properties */
+	MessageExpiryInterval = 2,		    /* 4 Byte Int :			Publish, Will Properties */
+	ContentType = 3,					/* Utf-8 String :		Publish, Will Properties */
+	ResponseTopic = 8,				    /* Utf-8 String :		Publish, Will Properties */
+	CorrelationData = 9,				/* Binary Data :		Publish, Will Properties */
+	SubscriptionIdentifier = 11,		/* Variable Byte Int :	Publish, Subscribe */
+	SessionExpiryInterval = 17,		    /* 4 Byte Int :			Connect, ConnectAcknowledge, Disconnect */
+	AssignedClientIdentifier = 18,	    /* Utf-8 String :		ConnectAcknowledge */
+	ServerKeepAlive = 19,			    /* 2 Byte Int :			ConnectAcknowledge */
+	AuthenticationMethod = 21,		    /* Utf-8 String :		Connect, ConnectAcknowledge, Authorize */
+	AuthenticationData = 22,			/* Binary Data :		Connect, ConnectAcknowledge, Authorize */
+	RequestProblemInformation = 23,	    /* Byte :				Connect */
+	WillDelayInterval = 24,			    /* 4 Byte Int :			Will Properties */
+	RequestResponseInformation = 25,    /* Byte :				Connect */
+	ResponseInformation = 26,		    /* Utf-8 String :		ConnectAcknowledge */
+	ServerReference = 28,			    /* Utf-8 String :		ConnectAcknowledge, Disconnect */
+	ReasonString = 31,				    /* Utf-8 String :		All Except Will Properties */
+	ReceiveMaximum = 33,				/* 2 Byte Int :			Connect, ConnectAcknowledge */
+	TopicAliasMaximum = 34,			    /* 2 Byte Int :			Connect, ConnectAcknowledge */
+	TopicAlias = 35,					/* 2 Byte Int :			Publish */
+	MaximumQoS = 36,					/* Byte :				ConnectAcknowledge */
+	RetainAvailable = 37,			    /* Byte :				ConnectAcknowledge */
+	UserProperty = 38,				    /* Utf-8 String Pair :	All */
+	MaximumPacketSize = 39,			    /* 4 Byte Int :			Connect, ConnectAcknowledge */
+	WildcardSubscriptionAvailable = 40, /* Byte :				ConnectAcknowledge */
+	SubscriptionIdAvailable = 41,	    /* Byte :				ConnectAcknowledge */
+	SharedSubscriptionAvailable = 42,	/* Byte :				ConnectAcknowledge */
+};
+
+enum class FMQTTPropertyType : uint8
+{
+	Byte = 1,
+	Int16 = 2,
+	Int32 = 3,
+	VarInt = 4,
+	Binary = 5,
+	String = 6,
+	StringPair = 7	
+};
+
+enum class FMQTTSubscriptionOption
+{
+	NoLocal = 0x04,				/* Won't receive own messages. */
+	RetainAsPublished = 0x08,	/* */
+	SendRetainAlways = 0x00,
+	SendRetainNew = 0x10,
+	SendRetainNever = 0x20	
+};
+
+/** 256MB MQTT v3.1.1: 2.2.3 */
+constexpr static uint32 GMQTTMaxPayload = 268435455U;
diff --git a/Engine/Plugins/Protocols/MQTT/Source/MQTTCore/Public/MQTTShared.h b/Engine/Plugins/Protocols/MQTT/Source/MQTTCore/Public/MQTTShared.h
new file mode 100644
--- /dev/null
+++ b/Engine/Plugins/Protocols/MQTT/Source/MQTTCore/Public/MQTTShared.h
@@ -0,0 +1,247 @@
+﻿// Copyright Epic Games, Inc. All Rights Reserved.
+
+#pragma once
+
+#include "CoreMinimal.h"
+#include "IPAddress.h"
+#include "MQTTProtocol.h"
+
+#include "MQTTShared.generated.h"
+
+struct FMQTTSubscription;
+class IMQTTClient;
+struct FMQTTClientMessage;
+class ISocketSubsystem;
+
+template <typename ResultType>
+class TCancellablePromise
+	: public TPromise<ResultType>
+{
+public:
+	~TCancellablePromise();
+	
+	bool IsStateValid()
+	{
+		return this->GetState().IsValid();
+	}
+};
+
+template <typename ResultType>
+TCancellablePromise<ResultType>::~TCancellablePromise()
+{
+	if(!IsStateValid())
+	{
+		this->SetValue({});
+	}
+}
+
+UENUM(BlueprintType)
+enum class EMQTTScheme : uint8
+{
+	MQTT = 0,
+	MQTTS = 1,		// Secure
+	
+	/* Possible but unsupported
+	TCP = 2,
+	SSL = 3
+	*/
+};
+
+UENUM(BlueprintType)
+enum class EMQTTQualityOfService : uint8
+{
+	Once = 0			UMETA(DisplayName = "QoS 0", Description = "Once (not guaranteed)"),
+	AtLeastOnce = 1		UMETA(DisplayName = "QoS 1", Description = "At Least Once (guaranteed)"),
+	ExactlyOnce = 2		UMETA(DisplayName = "QoS 2", Description = "Exactly Once (guaranteed)")
+};
+
+// mqtt[s]://[username][:password]@host.domain[:port]
+USTRUCT(BlueprintType)
+struct MQTTCORE_API FMQTTURL
+{
+	GENERATED_BODY()
+
+	UPROPERTY(BlueprintReadWrite, EditAnywhere, Category="URL")
+	FString Host = TEXT("localhost");
+
+	/** Default port is 1883 for MQTT scheme, 8883 for MQTTS. */
+	UPROPERTY(BlueprintReadWrite, EditAnywhere, AdvancedDisplay, Category="URL")
+	int32 Port = 1883;
+
+	UPROPERTY(BlueprintReadWrite, EditAnywhere, Category="URL")
+	FString Username;
+
+	UPROPERTY(BlueprintReadWrite, EditAnywhere, Category="URL", meta = (EditCondition = "Username"))
+	FString Password;
+
+	UPROPERTY(BlueprintReadWrite, EditAnywhere, AdvancedDisplay, Category="URL")
+	EMQTTScheme Scheme = EMQTTScheme::MQTT;
+
+	FMQTTURL() = default;
+
+	/** Construct from components. */
+	explicit FMQTTURL(
+		const FString& InHost,
+		const uint32 InPort = 1883,
+		const FString& InUsername = TEXT(""),
+		const FString& InPassword = TEXT(""),
+		const EMQTTScheme InScheme = EMQTTScheme::MQTT)
+		: Host(InHost),
+		  Port(InPort),
+		  Username(InUsername),
+		  Password(InPassword),
+		  Scheme(InScheme)
+	{
+	}
+
+	/** */
+	bool GetAddress(ISocketSubsystem* InSocketSubsystem, TSharedPtr<FInternetAddr>& OutAddress) const;
+
+	/** Deterministic FGuid based on the URL. */
+	FGuid ToGuid() const;
+
+	/** Gets the validity of the URL. */
+	bool IsValid() const;
+
+	/** Gets the validity of the URL with message. */
+	bool IsValid(FText& OutMessage) const;
+
+	/** Flags this URL as invalid. */
+	void SetInvalid() { bFlaggedInvalid = true; }
+
+	/** Parses the input string. */
+	static FMQTTURL FromString(const FString& InURL);
+
+	/** Get the URL as a string. */
+	FString ToString() const;
+
+protected:
+	bool bFlaggedInvalid = false;
+};
+
+inline uint32 GetTypeHash(const FMQTTURL& InURL)
+{
+	return GetTypeHash(InURL.ToString());
+}
+
+USTRUCT(BlueprintType)
+struct MQTTCORE_API FMQTTTopic
+{
+	GENERATED_BODY()
+
+	UPROPERTY(BlueprintReadWrite, EditAnywhere, Category="Topic")
+	FString Path;
+
+	FMQTTTopic() = default;
+
+	/** Construct from path. */
+	FMQTTTopic(const FString& InPath)
+		: Path(InPath)
+	{
+	}
+
+	virtual ~FMQTTTopic() = default;
+
+	/** Gets the validity of the input Topic. */
+	static bool IsValid(const FString& InTopic);
+	
+	/** Gets the validity of the Topic. */
+	virtual bool IsValid() const;
+
+	/** Attempts to make the topic valid. */
+	virtual void MakeValid();
+
+	/** Get the Topic as a string. */
+	virtual FString ToString() const;
+
+public:
+	static constexpr TCHAR TopicLevelSeparatorChar = '/'; // U+002F
+
+	/** '+' Must occupy entire level of the filter if it's used */
+	static constexpr TCHAR SingleLevelWildcardChar = '+'; // U+002B
+	
+	/** '#' Must be specified either on its own or following a topic level separator, must be the last character in the Topic Filter */
+	static constexpr TCHAR MultiLevelWildcardChar = '#'; // U+0023
+
+	/** Not part of spec, but widely used by servers. */
+	static constexpr TCHAR SystemChar = '$';
+};
+
+USTRUCT(BlueprintType)
+struct MQTTCORE_API FMQTTTopicFilter : public FMQTTTopic
+{
+	GENERATED_BODY()
+	
+	FMQTTTopicFilter() {  }
+	virtual ~FMQTTTopicFilter() override = default;
+
+	/** Gets the validity of the TopicFilter. */
+	virtual bool IsValid() const override;
+
+	/** Attempts to make the topic TopicFilter. */
+	virtual void MakeValid() override;
+
+	/** Get the TopicFilter as a string. */
+	virtual FString ToString() const override;
+};
+
+/** An MQTT Topic tokenized template */
+USTRUCT(BlueprintType)
+struct MQTTCORE_API FMQTTTopicTemplate : public FMQTTTopic
+{
+	GENERATED_BODY()
+
+	bool PopulateTemplate(const FStringFormatNamedArguments& InArgs, FMQTTTopic& OutTopic) const;
+	
+};
+
+struct MQTTCORE_API FMQTTSubscription
+{
+	/** Construct from path. */
+	FMQTTSubscription(const FMQTTTopic& InTopic, const EMQTTQualityOfService InGrantedQoS)
+		: Topic(InTopic)
+		, GrantedQoS(InGrantedQoS)
+	{
+	}
+
+	static bool Matches(const FString& InTopic, const FString& InSubscription);
+	bool Matches(const FString& InTopic) const;
+
+	const FMQTTTopic& GetTopic() const { return Topic; }
+
+	/** Gets the validity of the Subscription. */
+	bool IsValid() const;
+
+	DECLARE_MULTICAST_DELEGATE_OneParam(FOnSubscriptionMessage, const FMQTTClientMessage& /* InMessage */);
+	FOnSubscriptionMessage& OnSubscriptionMessage() { return OnSubscriptionMessageDelegate; }
+
+	EMQTTQualityOfService GetGrantedQoS() const { return GrantedQoS; }
+
+	friend uint32 GetTypeHash(const FMQTTSubscription& InSubscription)
+	{
+		return GetTypeHash(InSubscription.Topic.ToString());
+	}
+
+private:
+	friend uint32 GetTypeHash();
+	
+	FMQTTTopic Topic;
+	EMQTTQualityOfService GrantedQoS;
+	FOnSubscriptionMessage OnSubscriptionMessageDelegate;
+};
+
+struct MQTTCORE_API FMQTTSubscribeResult
+{
+	FMQTTSubscribeResult() = default;
+
+	FMQTTSubscribeResult(
+		const EMQTTSubscribeReturnCode InReturnCode,
+		TSharedPtr<FMQTTSubscription, ESPMode::ThreadSafe> InSubscription)
+		: ReturnCode(InReturnCode)
+		, Subscription(MoveTemp(InSubscription))
+	{
+	}
+
+	EMQTTSubscribeReturnCode ReturnCode;
+	TSharedPtr<FMQTTSubscription, ESPMode::ThreadSafe> Subscription;
+};
diff --git a/Engine/Plugins/Protocols/MQTT/Source/MQTTCore/Public/MQTTSubsystem.h b/Engine/Plugins/Protocols/MQTT/Source/MQTTCore/Public/MQTTSubsystem.h
new file mode 100644
--- /dev/null
+++ b/Engine/Plugins/Protocols/MQTT/Source/MQTTCore/Public/MQTTSubsystem.h
@@ -0,0 +1,41 @@
+﻿// Copyright Epic Games, Inc. All Rights Reserved.
+
+#pragma once
+
+#include "CoreMinimal.h"
+
+#include "MQTTClientMessage.h"
+#include "MQTTShared.h"
+#include "Subsystems/EngineSubsystem.h"
+#include "JsonObjectWrapper.h"
+
+#include "MQTTSubsystem.generated.h"
+
+class UMQTTServerObject;
+class UMQTTClientObject;
+
+/**
+* Provides persistent access to clients
+*/
+UCLASS(BlueprintType)
+class MQTTCORE_API UMQTTSubsystem final : public UEngineSubsystem
+{
+	GENERATED_BODY()
+
+public:
+	UFUNCTION(BlueprintCallable, Category = "MQTT", meta = (DefaultToSelf = "InParent", DisplayName = "Create Client (From Project URL)"))
+	static UMQTTClientObject* GetOrCreateClient_WithProjectURL(UPARAM(DisplayName = "Parent") UObject* InParent);
+	
+	UFUNCTION(BlueprintCallable, Category = "MQTT", meta = (DefaultToSelf = "InParent", DisplayName = "Create Client (From URL)"))
+	static UMQTTClientObject* GetOrCreateClient(UPARAM(DisplayName = "Parent") UObject* InParent, UPARAM(DisplayName = "URL") const FMQTTURL& InURL);
+
+	UFUNCTION(BlueprintCallable, BlueprintPure, Category = "MQTT")
+	static FString GetPayloadString(UPARAM(DisplayName = "ClientMessage") const FMQTTClientMessage& InClientMessage);
+
+	UFUNCTION(BlueprintCallable, BlueprintPure, Category = "MQTT", meta = (DefaultToSelf = "InParent"))
+	static bool GetPayloadJson(UPARAM(DisplayName = "Parent") UObject* InParent, UPARAM(DisplayName = "ClientMessage") const FMQTTClientMessage& InClientMessage, FJsonObjectWrapper& OutJson);
+
+private:
+	UPROPERTY(Transient)
+	TMap<FGuid, UMQTTClientObject*> Clients;
+};
diff --git a/Engine/Plugins/Protocols/MQTT/Source/MQTTCoreEditor/MQTTCoreEditor.Build.cs b/Engine/Plugins/Protocols/MQTT/Source/MQTTCoreEditor/MQTTCoreEditor.Build.cs
new file mode 100644
--- /dev/null
+++ b/Engine/Plugins/Protocols/MQTT/Source/MQTTCoreEditor/MQTTCoreEditor.Build.cs
@@ -0,0 +1,27 @@
+// Copyright Epic Games, Inc. All Rights Reserved.
+
+using UnrealBuildTool;
+
+public class MQTTCoreEditor : ModuleRules
+{
+	public MQTTCoreEditor(ReadOnlyTargetRules Target) : base(Target)
+	{
+		PCHUsage = PCHUsageMode.UseExplicitOrSharedPCHs;
+
+		PublicDependencyModuleNames.AddRange(
+			new[] {
+				"Core",
+                "CoreUObject"
+            });
+
+		PrivateDependencyModuleNames.AddRange(
+			new[] {
+				"Engine",
+				"MQTTCore",
+				"Settings",
+				"Slate",
+				"SlateCore",
+				"UnrealEd"
+			});
+    }
+}
diff --git a/Engine/Plugins/Protocols/MQTT/Source/MQTTCoreEditor/Private/MQTTCoreEditorLog.cpp b/Engine/Plugins/Protocols/MQTT/Source/MQTTCoreEditor/Private/MQTTCoreEditorLog.cpp
new file mode 100644
--- /dev/null
+++ b/Engine/Plugins/Protocols/MQTT/Source/MQTTCoreEditor/Private/MQTTCoreEditorLog.cpp
@@ -0,0 +1,5 @@
+// Copyright Epic Games, Inc. All Rights Reserved.
+
+#include "MQTTCoreEditorLog.h"
+
+DEFINE_LOG_CATEGORY(LogMQTTCoreEditor);
diff --git a/Engine/Plugins/Protocols/MQTT/Source/MQTTCoreEditor/Private/MQTTCoreEditorLog.h b/Engine/Plugins/Protocols/MQTT/Source/MQTTCoreEditor/Private/MQTTCoreEditorLog.h
new file mode 100644
--- /dev/null
+++ b/Engine/Plugins/Protocols/MQTT/Source/MQTTCoreEditor/Private/MQTTCoreEditorLog.h
@@ -0,0 +1,7 @@
+// Copyright Epic Games, Inc. All Rights Reserved.
+
+#pragma once
+
+#include "CoreMinimal.h"
+
+DECLARE_LOG_CATEGORY_EXTERN(LogMQTTCoreEditor, VeryVerbose, All);
diff --git a/Engine/Plugins/Protocols/MQTT/Source/MQTTCoreEditor/Private/MQTTCoreEditorModule.cpp b/Engine/Plugins/Protocols/MQTT/Source/MQTTCoreEditor/Private/MQTTCoreEditorModule.cpp
new file mode 100644
--- /dev/null
+++ b/Engine/Plugins/Protocols/MQTT/Source/MQTTCoreEditor/Private/MQTTCoreEditorModule.cpp
@@ -0,0 +1,37 @@
+// Copyright Epic Games, Inc. All Rights Reserved.
+
+#include "MQTTCoreEditorModule.h"
+
+#include "ISettingsModule.h"
+#include "MQTTClientSettings.h"
+
+#define LOCTEXT_NAMESPACE "MQTTCoreEditor"
+
+void FMQTTCoreEditorModule::StartupModule()
+{
+	ISettingsModule* SettingsModule = FModuleManager::GetModulePtr<ISettingsModule>("Settings");
+
+	// Register global settings
+	if (SettingsModule != nullptr)
+	{
+		SettingsModule->RegisterSettings("Project", "Plugins", "MQTT",
+										LOCTEXT("ProjectSettings_Label", "MQTT"),
+										LOCTEXT("ProjectSettings_Description", "Configure MQTT settings"),
+										GetMutableDefault<UMQTTClientSettings>()
+		);
+	}
+}
+
+void FMQTTCoreEditorModule::ShutdownModule()
+{
+	// Unregister global settings
+	ISettingsModule* SettingsModule = FModuleManager::GetModulePtr<ISettingsModule>("Settings");
+	if (SettingsModule)
+	{
+		SettingsModule->UnregisterSettings("Project", "Plugins", "MQTT");
+	}
+}
+
+#undef LOCTEXT_NAMESPACE
+
+IMPLEMENT_MODULE(FMQTTCoreEditorModule, MQTTCoreEditor)
diff --git a/Engine/Plugins/Protocols/MQTT/Source/MQTTCoreEditor/Private/MQTTCoreEditorModule.h b/Engine/Plugins/Protocols/MQTT/Source/MQTTCoreEditor/Private/MQTTCoreEditorModule.h
new file mode 100644
--- /dev/null
+++ b/Engine/Plugins/Protocols/MQTT/Source/MQTTCoreEditor/Private/MQTTCoreEditorModule.h
@@ -0,0 +1,16 @@
+// Copyright Epic Games, Inc. All Rights Reserved.
+
+#pragma once
+
+#include "IMQTTCoreEditorModule.h"
+
+#include "CoreMinimal.h"
+#include "Modules/ModuleManager.h"
+
+class FMQTTCoreEditorModule final
+    : public IMQTTCoreEditorModuleInterface
+{
+public:
+    virtual void StartupModule() override;
+    virtual void ShutdownModule() override;
+};
diff --git a/Engine/Plugins/Protocols/MQTT/Source/MQTTCoreEditor/Private/MQTTCoreEditorStats.cpp b/Engine/Plugins/Protocols/MQTT/Source/MQTTCoreEditor/Private/MQTTCoreEditorStats.cpp
new file mode 100644
--- /dev/null
+++ b/Engine/Plugins/Protocols/MQTT/Source/MQTTCoreEditor/Private/MQTTCoreEditorStats.cpp
@@ -0,0 +1,3 @@
+// Copyright Epic Games, Inc. All Rights Reserved.
+
+#include "MQTTCoreEditorStats.h"
diff --git a/Engine/Plugins/Protocols/MQTT/Source/MQTTCoreEditor/Private/MQTTCoreEditorStats.h b/Engine/Plugins/Protocols/MQTT/Source/MQTTCoreEditor/Private/MQTTCoreEditorStats.h
new file mode 100644
--- /dev/null
+++ b/Engine/Plugins/Protocols/MQTT/Source/MQTTCoreEditor/Private/MQTTCoreEditorStats.h
@@ -0,0 +1,10 @@
+// Copyright Epic Games, Inc. All Rights Reserved.
+
+#pragma once
+
+#include "CoreMinimal.h"
+
+#include "Stats/Stats.h"
+#include "Stats/Stats2.h"
+
+DECLARE_STATS_GROUP(TEXT("MQTTCoreEditor"), STATGROUP_MQTTCoreEditor, STATCAT_Advanced);
diff --git a/Engine/Plugins/Protocols/MQTT/Source/MQTTCoreEditor/Public/IMQTTCoreEditorModule.h b/Engine/Plugins/Protocols/MQTT/Source/MQTTCoreEditor/Public/IMQTTCoreEditorModule.h
new file mode 100644
--- /dev/null
+++ b/Engine/Plugins/Protocols/MQTT/Source/MQTTCoreEditor/Public/IMQTTCoreEditorModule.h
@@ -0,0 +1,11 @@
+// Copyright Epic Games, Inc. All Rights Reserved.
+
+#pragma once
+
+#include "CoreMinimal.h"
+#include "Modules/ModuleInterface.h"
+
+class IMQTTCoreEditorModuleInterface
+    : public IModuleInterface
+{
+};
