MQTT unchecked ptr fix

diff --git a/Engine/Plugins/Protocols/MQTT/Source/MQTTCore/Private/MQTTConnection.cpp b/Engine/Plugins/Protocols/MQTT/Source/MQTTCore/Private/MQTTConnection.cpp
--- a/Engine/Plugins/Protocols/MQTT/Source/MQTTCore/Private/MQTTConnection.cpp
+++ b/Engine/Plugins/Protocols/MQTT/Source/MQTTCore/Private/MQTTConnection.cpp
@@ -99,8 +99,8 @@ TSharedPtr<FMQTTConnection, ESPMode::ThreadSafe> FMQTTConnection::TryCreate(cons
 	
 	ISocketSubsystem* SocketSubsystem = ISocketSubsystem::Get(PLATFORM_SOCKETSUBSYSTEM);
 	FSocket* Socket = SocketSubsystem->CreateSocket(NAME_Stream, TEXT("MQTTConnection"), FNetworkProtocolTypes::IPv4);
-	Socket->SetReuseAddr(true);
 	check(Socket);
+	Socket->SetReuseAddr(true);
 
 	if (!Socket)
 	{
@@ -253,11 +253,8 @@ bool FMQTTConnection::UpdateConnection()
 
 void FMQTTConnection::ProcessMessages()
 {
-	UE_LOG(LogMQTTCore, Verbose, TEXT("BEEP: %s"), MQTT::GetMQTTStateName(GetState()));
-
 	if(GetState() == EMQTTState::Stopping)
 	{
-		UE_LOG(LogMQTTCore, Verbose, TEXT("BEEP was stopped."));
 		return;
 	}
 	
@@ -398,7 +395,7 @@ void FMQTTConnection::ProcessMessages()
 	}
 
 	// Wait until something to do
-	if (!Socket->Wait(ESocketWaitConditions::WaitForRead, FTimespan::FromSeconds(0.1f)))
+	if (Socket && !Socket->Wait(ESocketWaitConditions::WaitForRead, FTimespan::FromSeconds(0.1f)))
 	{
 		return;
 	}
@@ -433,8 +430,6 @@ void FMQTTConnection::ProcessMessages()
 			}
 		}
 	}
-
-	UE_LOG(LogMQTTCore, Verbose, TEXT("BOOP"));
 }
 
 void FMQTTConnection::OpenSocket()
diff --git a/Engine/Plugins/Protocols/MQTT/Source/MQTTCore/Public/MQTTClientObject.h b/Engine/Plugins/Protocols/MQTT/Source/MQTTCore/Public/MQTTClientObject.h
--- a/Engine/Plugins/Protocols/MQTT/Source/MQTTCore/Public/MQTTClientObject.h
+++ b/Engine/Plugins/Protocols/MQTT/Source/MQTTCore/Public/MQTTClientObject.h
@@ -24,7 +24,7 @@ public:
 	void Initialize(const TSharedPtr<FMQTTSubscription, ESPMode::ThreadSafe>& InSubscription);
 
 	UFUNCTION(BlueprintCallable, Category = "MQTT")
-	virtual void SetOnMessageHandler(const FOnMessageDelegate& OnMessageCallback);
+	void SetOnMessageHandler(const FOnMessageDelegate& OnMessageCallback);
 
 	UFUNCTION(BlueprintCallable, BlueprintPure, Category = "MQTT")
 	bool IsValid() const { return MqttSubscriptionPtr.IsValid(); }
