MQTT cis fix

diff --git a/Engine/Plugins/Protocols/MQTT/Source/MQTTCore/Private/MQTTConnection.cpp b/Engine/Plugins/Protocols/MQTT/Source/MQTTCore/Private/MQTTConnection.cpp
--- a/Engine/Plugins/Protocols/MQTT/Source/MQTTCore/Private/MQTTConnection.cpp
+++ b/Engine/Plugins/Protocols/MQTT/Source/MQTTCore/Private/MQTTConnection.cpp
@@ -99,8 +99,6 @@ TSharedPtr<FMQTTConnection, ESPMode::ThreadSafe> FMQTTConnection::TryCreate(cons
 	
 	ISocketSubsystem* SocketSubsystem = ISocketSubsystem::Get(PLATFORM_SOCKETSUBSYSTEM);
 	FSocket* Socket = SocketSubsystem->CreateSocket(NAME_Stream, TEXT("MQTTConnection"), FNetworkProtocolTypes::IPv4);
-	check(Socket);
-	Socket->SetReuseAddr(true);
 
 	if (!Socket)
 	{
@@ -108,6 +106,8 @@ TSharedPtr<FMQTTConnection, ESPMode::ThreadSafe> FMQTTConnection::TryCreate(cons
 		return nullptr;
 	}
 
+	Socket->SetReuseAddr(true);
+
 	TSharedPtr<FMQTTConnection, ESPMode::ThreadSafe> Runnable = MakeShared<FMQTTConnection, ESPMode::ThreadSafe>(*Socket, InAddr.ToSharedRef());
 	return Runnable;
 }
