MQTT string scope fix

diff --git a/Engine/Plugins/Protocols/MQTT/Source/MQTTCore/Private/MQTTOperations.cpp b/Engine/Plugins/Protocols/MQTT/Source/MQTTCore/Private/MQTTOperations.cpp
--- a/Engine/Plugins/Protocols/MQTT/Source/MQTTCore/Private/MQTTOperations.cpp
+++ b/Engine/Plugins/Protocols/MQTT/Source/MQTTCore/Private/MQTTOperations.cpp
@@ -59,9 +59,9 @@ template <typename SizeType = uint16>
 void SerializeString(FArchive& Ar, const FString& InString)
 {
 	SizeType StringLength = InString.Len();
-	const char* String = TCHAR_TO_UTF8(*InString);
+	const auto String = StringCast<ANSICHAR>(*InString);
 	Ar << StringLength;
-	Ar.Serialize((void*)String, StringLength);
+	Ar.Serialize((void*)String.Get(), StringLength);
 }
 
 template <typename SizeType = uint16>
