Changed MQTT Plugin from Beta To Experimental #jira none #rb simon.therriault #preflight none

diff --git a/Engine/Plugins/Protocols/MQTT/MQTT.uplugin b/Engine/Plugins/Protocols/MQTT/MQTT.uplugin
--- a/Engine/Plugins/Protocols/MQTT/MQTT.uplugin
+++ b/Engine/Plugins/Protocols/MQTT/MQTT.uplugin
@@ -9,7 +9,8 @@
 	"CreatedByURL": "https://epicgames.com",
 	"EnabledByDefault": false,
 	"CanContainContent": false,
-	"IsBetaVersion": true,
+	"IsBetaVersion": false,
+	"IsExperimentalVersion": true,
 	"SupportedTargetPlatforms": [
 		"Win64",
 		"Mac",
