Deleted unused TCancellablePromise, in advance of making TPromise final.

diff --git a/Engine/Plugins/Protocols/MQTT/Source/MQTTCore/Public/IMQTTClient.h b/Engine/Plugins/Protocols/MQTT/Source/MQTTCore/Public/IMQTTClient.h
--- a/Engine/Plugins/Protocols/MQTT/Source/MQTTCore/Public/IMQTTClient.h
+++ b/Engine/Plugins/Protocols/MQTT/Source/MQTTCore/Public/IMQTTClient.h
@@ -1,4 +1,4 @@
-﻿// Copyright Epic Games, Inc. All Rights Reserved.
+// Copyright Epic Games, Inc. All Rights Reserved.
 
 #pragma once
 
@@ -39,8 +39,7 @@ public:
 	// @todo: store session
 	/** Connect using the clients URL. Set bCleanSession = false to resume previous session. */
 	virtual TFuture<EMQTTConnectReturnCode> Connect(bool bCleanSession = true) = 0;
-	
-	//virtual TCancellablePromise<bool> Connect();
+
 	virtual TFuture<void> Disconnect() = 0;
 
 	/** Return true if the entire operation was successful (depends on QoS) */
diff --git a/Engine/Plugins/Protocols/MQTT/Source/MQTTCore/Public/MQTTShared.h b/Engine/Plugins/Protocols/MQTT/Source/MQTTCore/Public/MQTTShared.h
--- a/Engine/Plugins/Protocols/MQTT/Source/MQTTCore/Public/MQTTShared.h
+++ b/Engine/Plugins/Protocols/MQTT/Source/MQTTCore/Public/MQTTShared.h
@@ -14,28 +14,6 @@ class IMQTTClient;
 struct FMQTTClientMessage;
 class ISocketSubsystem;
 
-template <typename ResultType>
-class TCancellablePromise
-	: public TPromise<ResultType>
-{
-public:
-	~TCancellablePromise();
-	
-	bool IsStateValid()
-	{
-		return this->GetState().IsValid();
-	}
-};
-
-template <typename ResultType>
-TCancellablePromise<ResultType>::~TCancellablePromise()
-{
-	if(!IsStateValid())
-	{
-		this->SetValue({});
-	}
-}
-
 UENUM(BlueprintType)
 enum class EMQTTScheme : uint8
 {
