[MQTT] Removed platform restrictions

diff --git a/Engine/Plugins/Protocols/MQTT/MQTT.uplugin b/Engine/Plugins/Protocols/MQTT/MQTT.uplugin
--- a/Engine/Plugins/Protocols/MQTT/MQTT.uplugin
+++ b/Engine/Plugins/Protocols/MQTT/MQTT.uplugin
@@ -11,11 +11,6 @@
 	"CanContainContent": false,
 	"IsBetaVersion": false,
 	"IsExperimentalVersion": true,
-	"SupportedTargetPlatforms": [
-		"Win64",
-		"Mac",
-		"Linux"
-	],
 	"Modules": [
 		{
 			"Name": "MQTTCore",
