Fixed a crash that occurred when receiving MQTT packets with payloads =128 bytes due to incorrect variable-length header size calculation.

#Synced-CL 50301123

diff --git a/Engine/Plugins/Protocols/MQTT/Source/MQTTCore/Private/MQTTOperations.cpp b/Engine/Plugins/Protocols/MQTT/Source/MQTTCore/Private/MQTTOperations.cpp
--- a/Engine/Plugins/Protocols/MQTT/Source/MQTTCore/Private/MQTTOperations.cpp
+++ b/Engine/Plugins/Protocols/MQTT/Source/MQTTCore/Private/MQTTOperations.cpp
@@ -274,35 +274,33 @@ auto FMQTTPublishPacket::Serialize(FArchive& Ar) -> void
 
 	if(Ar.IsLoading())
 	{
-		// Mark read position here
-		int64 HeaderSize = Ar.Tell();
 		Ar << FixedHeader;
 
 		bIsDuplicate = (FixedHeader.PacketType & 0x08) >> 3;
 		QoS = static_cast<EMQTTQualityOfService>((FixedHeader.PacketType & 0x06) >> 1);
 		bRetain = (FixedHeader.PacketType & 0x01);
 
+		const int64 HeaderEndPosition = Ar.Tell();
+
 		DeserializeString(Ar, Topic);
 
 		if(QoS > EMQTTQualityOfService::Once)
 		{
 			Ar << PacketId;
 		}
-		// Calculate size
-		HeaderSize = Ar.Tell() - HeaderSize;
 
-		int32 PayloadLength = FixedHeader.RemainingLength - HeaderSize;
+		// Calculate variable header size (topic + optional packet id)
+		const int64 VariableHeaderSize = Ar.Tell() - HeaderEndPosition;
+		int32 PayloadLength = FixedHeader.RemainingLength - VariableHeaderSize;
+
 		if(PayloadLength > 0)
 		{
-			PayloadLength += 2;
 			Payload.Init(0, PayloadLength);
 			for(int Idx = 0; Idx < PayloadLength; ++Idx)
 			{
 				Ar << Payload[Idx];
 			}
 		}
-		auto X = Ar.Tell();
-		auto Y = X;
 	}
 	else if(Ar.IsSaving())
 	{
